#summary 	Fixed opening multiple instances of visual studio
#ttp 330039 	EDITOR: Platform-agnostic editor code depends on Windows-only VSAccessor headers
#detail 	Merge error meant that the 'open' timer wasnt getting set properly.

diff --git a/Engine/Plugins/Developer/VisualStudioSourceCodeAccess/Source/VisualStudioSourceCodeAccess/Private/VisualStudioSourceCodeAccessor.cpp b/Engine/Plugins/Developer/VisualStudioSourceCodeAccess/Source/VisualStudioSourceCodeAccess/Private/VisualStudioSourceCodeAccessor.cpp
--- a/Engine/Plugins/Developer/VisualStudioSourceCodeAccess/Source/VisualStudioSourceCodeAccess/Private/VisualStudioSourceCodeAccessor.cpp
+++ b/Engine/Plugins/Developer/VisualStudioSourceCodeAccess/Source/VisualStudioSourceCodeAccess/Private/VisualStudioSourceCodeAccessor.cpp
@@ -303,7 +303,7 @@ bool FVisualStudioSourceCodeAccessor::OpenVisualStudioFilesInternalViaDTE(const
 					}
 				}
 
-				SourceCodeAccessModule.OnDoneLaunchingCodeAccessor().Broadcast( true );
+				VSLaunchFinished( true );
 			}
 			else
 			{
@@ -335,7 +335,7 @@ bool FVisualStudioSourceCodeAccessor::OpenVisualStudioFilesInternalViaDTE(const
 				}
 				else
 				{
-					SourceCodeAccessModule.OnLaunchingCodeAccessor().Broadcast();
+					VSLaunchStarted();
 				}
 			}
 		}
@@ -356,7 +356,7 @@ bool FVisualStudioSourceCodeAccessor::OpenVisualStudioFilesInternalViaDTE(const
 		{
 			// We need todo this incase the process died or was kill prior to the code gaining focus of it
 			bDefer = false;
-			SourceCodeAccessModule.OnDoneLaunchingCodeAccessor().Broadcast( false );
+			VSLaunchFinished(false);
 
 			// We failed to open the solution and file, so lets just use the platforms default opener.
 			for ( const FileOpenRequest& Request : Requests )
