[INTEGRATE] Engine source and content from 4.1 branch.

diff --git a/Engine/Plugins/Developer/VisualStudioSourceCodeAccess/Source/VisualStudioSourceCodeAccess/Private/VisualStudioSourceCodeAccessor.cpp b/Engine/Plugins/Developer/VisualStudioSourceCodeAccess/Source/VisualStudioSourceCodeAccess/Private/VisualStudioSourceCodeAccessor.cpp
--- a/Engine/Plugins/Developer/VisualStudioSourceCodeAccess/Source/VisualStudioSourceCodeAccess/Private/VisualStudioSourceCodeAccessor.cpp
+++ b/Engine/Plugins/Developer/VisualStudioSourceCodeAccess/Source/VisualStudioSourceCodeAccess/Private/VisualStudioSourceCodeAccessor.cpp
@@ -69,7 +69,7 @@ void FVisualStudioSourceCodeAccessor::Startup()
 	FModuleManager::Get().OnModuleCompilerStarted().AddStatic( &SaveVisualStudioDocuments );
 
 	// Cache this so we don't have to do it on a background thread
-	SolutionPath = FPaths::ConvertRelativePathToFull(FModuleManager::Get().GetSolutionFilepath());
+	GetSolutionPath();
 
 	// Preferential order of VS versions
 	AddVisualStudioVersion(12); // Visual Studio 2013
@@ -111,37 +111,37 @@ EAccessVisualStudioResult AccessVisualStudioViaDTE(CComPtr<EnvDTE::_DTE>& OutDTE
 		IEnumMoniker* MonikersTable;
 		if(SUCCEEDED(RunningObjectTable->EnumRunning(&MonikersTable)))
 		{
-			MonikersTable->Reset();
+		MonikersTable->Reset();
 
-			// Look for all visual studio instances in the ROT
-			IMoniker* CurrentMoniker;
+		// Look for all visual studio instances in the ROT
+		IMoniker* CurrentMoniker;
 			while(AccessResult != EAccessVisualStudioResult::VSInstanceIsOpen && MonikersTable->Next(1, &CurrentMoniker, NULL) == S_OK)
-			{
-				IBindCtx* BindContext;
-				LPOLESTR OutName;
+		{
+			IBindCtx* BindContext;
+			LPOLESTR OutName;
 				if(SUCCEEDED(CreateBindCtx(0, &BindContext)) && SUCCEEDED(CurrentMoniker->GetDisplayName(BindContext, NULL, &OutName)))
 				{
 					if(IsVisualStudioDTEMoniker(FString(OutName), InLocations))
 					{
-						CComPtr<IUnknown> ComObject;
+			CComPtr<IUnknown> ComObject;
 						if(SUCCEEDED(RunningObjectTable->GetObject(CurrentMoniker, &ComObject)))
-						{
-							CComPtr<EnvDTE::_DTE> TempDTE;
-							TempDTE = ComObject;
-
-							// Get the solution path for this instance
-							// If it equals the solution we would have opened above in RunVisualStudio(), we'll take that
-							CComPtr<EnvDTE::_Solution> Solution;
-							LPOLESTR OutPath;
-							if(SUCCEEDED(TempDTE->get_Solution(&Solution)) &&
-								SUCCEEDED(Solution->get_FullName(&OutPath)))
-							{
-								FString Filename(OutPath);
-								FPaths::NormalizeFilename(Filename);
+			{
+				CComPtr<EnvDTE::_DTE> TempDTE;
+				TempDTE = ComObject;
+
+				// Get the solution path for this instance
+				// If it equals the solution we would have opened above in RunVisualStudio(), we'll take that
+				CComPtr<EnvDTE::_Solution> Solution;
+				LPOLESTR OutPath;
+				if (SUCCEEDED(TempDTE->get_Solution(&Solution)) &&
+					SUCCEEDED(Solution->get_FullName(&OutPath)))
+				{
+					FString Filename(OutPath);
+					FPaths::NormalizeFilename(Filename);
 
-								if(Filename == InSolutionPath)
-								{
-									OutDTE = TempDTE;
+					if( Filename == InSolutionPath )
+					{
+						OutDTE = TempDTE;
 									AccessResult = EAccessVisualStudioResult::VSInstanceIsOpen;
 								}
 							}
@@ -157,11 +157,11 @@ EAccessVisualStudioResult AccessVisualStudioViaDTE(CComPtr<EnvDTE::_DTE>& OutDTE
 				{
 					UE_LOG(LogVSAccessor, Warning, TEXT("Couldn't get display name"));
 					AccessResult = EAccessVisualStudioResult::VSInstanceUnknown;
-				}
-				BindContext->Release();
-				CurrentMoniker->Release();
 			}
-			MonikersTable->Release();
+			BindContext->Release();
+			CurrentMoniker->Release();
+		}
+		MonikersTable->Release();
 		}
 		else
 		{
@@ -191,32 +191,32 @@ bool FVisualStudioSourceCodeAccessor::OpenVisualStudioSolutionViaDTE()
 	bool bSuccess = false;
 
 	CComPtr<EnvDTE::_DTE> DTE;
-	switch (AccessVisualStudioViaDTE(DTE, SolutionPath, Locations))
+	switch (AccessVisualStudioViaDTE(DTE, GetSolutionPath(), Locations))
 	{
 		case EAccessVisualStudioResult::VSInstanceIsOpen:
+	{
+		// Set Focus on Visual Studio
+		CComPtr<EnvDTE::Window> MainWindow;
+		if (SUCCEEDED(DTE->get_MainWindow(&MainWindow)) &&
+			SUCCEEDED(MainWindow->Activate()))
 		{
-			// Set Focus on Visual Studio
-			CComPtr<EnvDTE::Window> MainWindow;
-			if (SUCCEEDED(DTE->get_MainWindow(&MainWindow)) &&
-				SUCCEEDED(MainWindow->Activate()))
-			{
-				bSuccess = true;
-			}
-			else
-			{
-				UE_LOG(LogVSAccessor, Warning, TEXT("Couldn't set focus on Visual Studio."));
-			}
+			bSuccess = true;
+		}
+		else
+		{
+			UE_LOG(LogVSAccessor, Warning, TEXT("Couldn't set focus on Visual Studio."));
 		}
+	}
 		break;
 
 	case EAccessVisualStudioResult::VSInstanceIsNotOpen:
+	{
+		// Automatically fail if there's already an attempt in progress
+		if (!IsVSLaunchInProgress())
 		{
-			// Automatically fail if there's already an attempt in progress
-			if (!IsVSLaunchInProgress())
-			{
-				bSuccess = RunVisualStudioAndOpenSolution();
-			}
+			bSuccess = RunVisualStudioAndOpenSolution();
 		}
+	}
 		break;
 
 	default:
@@ -244,7 +244,7 @@ bool FVisualStudioSourceCodeAccessor::OpenVisualStudioFilesInternalViaDTE(const
 	
 	bool bDefer = false, bSuccess = false;
 	CComPtr<EnvDTE::_DTE> DTE;
-	switch (AccessVisualStudioViaDTE(DTE, SolutionPath, Locations))
+	switch (AccessVisualStudioViaDTE(DTE, GetSolutionPath(), Locations))
 	{
 	case EAccessVisualStudioResult::VSInstanceIsOpen:
 	{
@@ -272,7 +272,7 @@ bool FVisualStudioSourceCodeAccessor::OpenVisualStudioFilesInternalViaDTE(const
 					CComBSTR COMStrFileName(ANSIPath.Get());
 					CComBSTR COMStrKind(EnvDTE::vsViewKindTextView);
 					CComPtr<EnvDTE::Window> Window;
-					if (SUCCEEDED(ItemOperations->OpenFile(COMStrFileName, COMStrKind, &Window)))
+					if ( SUCCEEDED(ItemOperations->OpenFile(COMStrFileName, COMStrKind, &Window)) )
 					{
 						// If we've made it this far - we've opened the file.  it doesn't matter if
 						// we successfully get to the line number.  Everything else is gravy.
@@ -282,10 +282,10 @@ bool FVisualStudioSourceCodeAccessor::OpenVisualStudioFilesInternalViaDTE(const
 						CComPtr<EnvDTE::Document> Document;
 						CComPtr<IDispatch> SelectionDispatch;
 						CComPtr<EnvDTE::TextSelection> Selection;
-						if (SUCCEEDED(DTE->get_ActiveDocument(&Document)) &&
-							SUCCEEDED(Document->get_Selection(&SelectionDispatch)) &&
-							SUCCEEDED(SelectionDispatch->QueryInterface(&Selection)) &&
-							SUCCEEDED(Selection->GotoLine(Request.LineNumber, true)))
+						if ( SUCCEEDED(DTE->get_ActiveDocument(&Document)) &&
+							 SUCCEEDED(Document->get_Selection(&SelectionDispatch)) &&
+							 SUCCEEDED(SelectionDispatch->QueryInterface(&Selection)) &&
+							 SUCCEEDED(Selection->GotoLine(Request.LineNumber, true)) )
 						{
 							if ( !SUCCEEDED(Selection->MoveToLineAndOffset(Request.LineNumber, Request.ColumnNumber, false)) )
 							{
@@ -326,7 +326,7 @@ bool FVisualStudioSourceCodeAccessor::OpenVisualStudioFilesInternalViaDTE(const
 		if ( IsInGameThread() )
 		{
 			// If we haven't already attempted to launch VS do so now
-			if (!IsVSLaunchInProgress())
+			if ( !IsVSLaunchInProgress() )
 			{
 				// If there's no valid instance of VS running, run one if we have it installed
 				if ( !RunVisualStudioAndOpenSolution() )
@@ -348,13 +348,13 @@ bool FVisualStudioSourceCodeAccessor::OpenVisualStudioFilesInternalViaDTE(const
 		bDefer = false;
 		break;
 	}
-
+	
 	if ( !bSuccess )
 	{
 		// If we have attempted to launch VS, and it's taken too long, timeout so the user can try again
 		if ( IsVSLaunchInProgress() && ( FPlatformTime::Seconds() - VSLaunchTime ) > 300 )
 		{
-			// We need todo this incase the process died or was kill prior to the code gaining focus of it
+			// We need todo this in case the process died or was kill prior to the code gaining focus of it
 			bDefer = false;
 			VSLaunchFinished(false);
 
@@ -399,7 +399,7 @@ bool FVisualStudioSourceCodeAccessor::SaveAllOpenDocuments() const
 	}
 	
 	CComPtr<EnvDTE::_DTE> DTE;
-	if (AccessVisualStudioViaDTE(DTE, SolutionPath, Locations) == EAccessVisualStudioResult::VSInstanceIsOpen)
+	if (AccessVisualStudioViaDTE(DTE, GetSolutionPath(), Locations) == EAccessVisualStudioResult::VSInstanceIsOpen)
 	{
 		// Save all documents
 		CComPtr<EnvDTE::Documents> Documents;
@@ -636,8 +636,8 @@ EAccessVisualStudioResult AccessVisualStudioViaProcess(::DWORD& OutProcessID, FS
 							{
 								UE_LOG(LogVSAccessor, Warning, TEXT("Couldn't access module information"));
 								AccessResult = EAccessVisualStudioResult::VSInstanceUnknown;
-							}
 						}
+					}
 						}
 					else
 					{
@@ -695,25 +695,25 @@ bool FVisualStudioSourceCodeAccessor::OpenVisualStudioSolutionViaProcess()
 {
 	::DWORD ProcessID = 0;
 	FString Path;
-	switch (AccessVisualStudioViaProcess(ProcessID, Path, SolutionPath, Locations))
+	switch (AccessVisualStudioViaProcess(ProcessID, Path, GetSolutionPath(), Locations))
 {
 	case EAccessVisualStudioResult::VSInstanceIsOpen:
+	{
+		// Try to bring Visual Studio to the foreground
+		::HWND VisualStudioHwnd = GetTopWindowForProcess(ProcessID);
+		if (VisualStudioHwnd)
 		{
-			// Try to bring Visual Studio to the foreground
-			::HWND VisualStudioHwnd = GetTopWindowForProcess(ProcessID);
-			if (VisualStudioHwnd)
-			{
-				// SwitchToThisWindow isn't really intended for general use, however it can switch to 
-				// the VS window, where SetForegroundWindow will fail due to process permissions
-				::SwitchToThisWindow(VisualStudioHwnd, 0);
-			}
-			return true;
+			// SwitchToThisWindow isn't really intended for general use, however it can switch to 
+			// the VS window, where SetForegroundWindow will fail due to process permissions
+			::SwitchToThisWindow(VisualStudioHwnd, 0);
 		}
+		return true;
+	}
 		break;
 
 	case EAccessVisualStudioResult::VSInstanceIsNotOpen:
-		return RunVisualStudioAndOpenSolution();
-		
+	return RunVisualStudioAndOpenSolution();
+
 	default:
 		// Do nothing if we failed the VS detection, otherwise we could get stuck in a loop of constantly 
 		// trying to open a VS instance since we can't detect that one is already running
@@ -721,13 +721,13 @@ bool FVisualStudioSourceCodeAccessor::OpenVisualStudioSolutionViaProcess()
 	}
 
 		return false;
-	}
+}
 
 bool FVisualStudioSourceCodeAccessor::OpenVisualStudioFilesInternalViaProcess(const TArray<FileOpenRequest>& Requests)
 {
 	::DWORD ProcessID = 0;
 	FString Path;
-	switch (AccessVisualStudioViaProcess(ProcessID, Path, SolutionPath, Locations))
+	switch (AccessVisualStudioViaProcess(ProcessID, Path, GetSolutionPath(), Locations))
 	{
 	case EAccessVisualStudioResult::VSInstanceIsOpen:
 		return RunVisualStudioAndOpenSolutionAndFiles(Path, "", &Requests);
@@ -735,7 +735,7 @@ bool FVisualStudioSourceCodeAccessor::OpenVisualStudioFilesInternalViaProcess(co
 	case EAccessVisualStudioResult::VSInstanceIsNotOpen:
 		if (CanRunVisualStudio(Path))
 		{
-			return RunVisualStudioAndOpenSolutionAndFiles(Path, SolutionPath, &Requests);
+			return RunVisualStudioAndOpenSolutionAndFiles(Path, GetSolutionPath(), &Requests);
 		}
 		break;
 
@@ -766,7 +766,7 @@ bool FVisualStudioSourceCodeAccessor::RunVisualStudioAndOpenSolution() const
 	FString Path;
 	if (CanRunVisualStudio(Path))
 	{
-		return RunVisualStudioAndOpenSolutionAndFiles(Path, SolutionPath, nullptr);
+		return RunVisualStudioAndOpenSolutionAndFiles(Path, GetSolutionPath(), nullptr);
 	}
 
 	return false;
@@ -797,7 +797,7 @@ bool FVisualStudioSourceCodeAccessor::OpenFileAtLine(const FString& FullPath, in
 		LineNumber++;
 	}
 	if(ColumnNumber == 0)
-	{
+{
 		ColumnNumber++;
 	}
 
@@ -947,4 +947,14 @@ FText FVisualStudioSourceCodeAccessor::GetDescriptionText() const
 	return LOCTEXT("VisualStudioDisplayDesc", "Open source code files in Visual Studio");
 }
 
+FString FVisualStudioSourceCodeAccessor::GetSolutionPath() const
+{
+	FScopeLock Lock(&CachedSolutionPathCriticalSection);
+	if(IsInGameThread())
+	{
+		CachedSolutionPath = FPaths::ConvertRelativePathToFull(FModuleManager::Get().GetSolutionFilepath());
+	}
+	return CachedSolutionPath;
+}
+
 #undef LOCTEXT_NAMESPACE
\ No newline at end of file
diff --git a/Engine/Plugins/Developer/VisualStudioSourceCodeAccess/Source/VisualStudioSourceCodeAccess/Private/VisualStudioSourceCodeAccessor.h b/Engine/Plugins/Developer/VisualStudioSourceCodeAccess/Source/VisualStudioSourceCodeAccess/Private/VisualStudioSourceCodeAccessor.h
--- a/Engine/Plugins/Developer/VisualStudioSourceCodeAccess/Source/VisualStudioSourceCodeAccess/Private/VisualStudioSourceCodeAccessor.h
+++ b/Engine/Plugins/Developer/VisualStudioSourceCodeAccess/Source/VisualStudioSourceCodeAccess/Private/VisualStudioSourceCodeAccessor.h
@@ -65,10 +65,10 @@ private:
 
 	/** Opens a file in the correct running instance of Visual Studio at a line and optionally to a column. */
 	bool OpenVisualStudioFileAtLineInternal(const FString& FullPath, int32 LineNumber, int32 ColumnNumber = 0);
-		
+
 	/** Opens multiple files in the correct running instance of Visual Studio. */
 	bool OpenVisualStudioFilesInternal(const TArray<FileOpenRequest>& Requests);
-
+		
 	/** An instance of VS it attempting to be opened */
 	void VSLaunchStarted();
 
@@ -118,8 +118,14 @@ private:
 	TArray<VisualStudioLocation> Locations;
 
 	/** String storing the solution path obtained from the module manager to avoid having to use it on a thread */
-	FString SolutionPath;
+	mutable FString CachedSolutionPath;
+
+	/** Critical section for updating SolutionPath */
+	mutable FCriticalSection CachedSolutionPathCriticalSection;
 
 	/** If !0 it represents the time at which the a VS instance was opened */
 	double	VSLaunchTime;
-};
\ No newline at end of file
+
+	/** Accessor for SolutionPath. Will try to update it when called from the game thread, otherwise will use the cached value */
+	FString GetSolutionPath() const;
+};
diff --git a/Engine/Source/Editor/GameProjectGeneration/Private/GameProjectUtils.cpp b/Engine/Source/Editor/GameProjectGeneration/Private/GameProjectUtils.cpp
--- a/Engine/Source/Editor/GameProjectGeneration/Private/GameProjectUtils.cpp
+++ b/Engine/Source/Editor/GameProjectGeneration/Private/GameProjectUtils.cpp
@@ -264,11 +264,9 @@ void GameProjectUtils::GetStarterContentFiles(TArray<FString>& OutFilenames)
 {
 	FString const SrcFolder = FPaths::StarterContentDir();
 	FString const ContentFolder = SrcFolder / TEXT("Content");
-	FString const DDCFolder = SrcFolder / TEXT("DerivedDataCache");
 
-	// only copying /Content and /DerivedDataCache
+	// only copying /Content
 	IFileManager::Get().FindFilesRecursive(OutFilenames, *ContentFolder, TEXT("*"), /*Files=*/true, /*Directories=*/false);
-	IFileManager::Get().FindFilesRecursive(OutFilenames, *DDCFolder, TEXT("*"), /*Files=*/true, /*Directories=*/false, /*bClearFilenames=*/false);
 }
 
 bool GameProjectUtils::CopyStarterContent(const FString& DestProjectFolder, FText& OutFailReason)
diff --git a/Engine/Source/Programs/AutomationTool/IOS/IOSPlatform.Automation.cs b/Engine/Source/Programs/AutomationTool/IOS/IOSPlatform.Automation.cs
--- a/Engine/Source/Programs/AutomationTool/IOS/IOSPlatform.Automation.cs
+++ b/Engine/Source/Programs/AutomationTool/IOS/IOSPlatform.Automation.cs
@@ -226,7 +226,7 @@ public class IOSPlatform : Platform
 				Directory.SetCurrentDirectory (CWD);
 
 				XcodeProj = RawProjectPath.Replace(".uproject", ".xcodeproj"); 
-				if (!Directory.Exists(XcodeProj))
+				if (!Directory.Exists (XcodeProj))
 				{
 					// something very bad happened
 					throw new AutomationException("iOS couldn't find the appropriate Xcode Project");
@@ -280,6 +280,10 @@ public class IOSPlatform : Platform
 		// create the file
 		using (ZipFile Zip = new ZipFile())
 		{
+            // Set encoding to support unicode filenames
+            Zip.AlternateEncodingUsage = ZipOption.Always;
+            Zip.AlternateEncoding = Encoding.UTF8;
+
 			// set the compression level
 			if (Distribution)
 			{
diff --git a/Engine/Source/Programs/AutomationTool/Linux/LinuxPlatform.Automation.cs b/Engine/Source/Programs/AutomationTool/Linux/LinuxPlatform.Automation.cs
--- a/Engine/Source/Programs/AutomationTool/Linux/LinuxPlatform.Automation.cs
+++ b/Engine/Source/Programs/AutomationTool/Linux/LinuxPlatform.Automation.cs
@@ -31,11 +31,17 @@ public abstract class BaseLinuxPlatform : Platform
         }
 
         // assume that we always have to deploy Steam (FIXME: should be automatic)
-        {
-            string SteamVersion = "Steamv128";
-            SC.StageFiles(StagedFileType.NonUFS, CombinePaths(SC.LocalRoot, "Engine/Binaries/ThirdParty/Steamworks/" + SteamVersion, SC.PlatformDir), "libsteam_api.so", false, null, CombinePaths("Engine/Binaries", SC.PlatformDir));
-            SC.StageFiles(StagedFileType.NonUFS, CombinePaths(SC.LocalRoot, "Engine/Config"), "controller.vdf", false, null, CommandUtils.CombinePaths(SC.RelativeProjectRootForStage, "Saved/Config"));
-        }
+		{
+			string SteamVersion = "Steamv128";
+
+			// Check if the Steam directory exists. We need it for Steam controller support, so we include it whenever we can.
+			if (Directory.Exists(CommandUtils.CombinePaths(SC.LocalRoot, "Engine/Binaries/ThirdParty/Steamworks/" + SteamVersion)))
+			{
+				SC.StageFiles(StagedFileType.NonUFS, CombinePaths(SC.LocalRoot, "Engine/Binaries/ThirdParty/Steamworks/" + SteamVersion, SC.PlatformDir), "libsteam_api.so", false, null, CombinePaths("Engine/Binaries", SC.PlatformDir));
+			}
+
+			SC.StageFiles(StagedFileType.NonUFS, CombinePaths(SC.LocalRoot, "Engine/Config"), "controller.vdf", false, null, CommandUtils.CombinePaths(SC.RelativeProjectRootForStage, "Saved/Config"));
+		}
 
         // assume that we always have to deploy OpenAL (FIXME: should be automatic)
         {
diff --git a/Engine/Source/Programs/AutomationTool/Scripts/BuildDocumentation.Automation.cs b/Engine/Source/Programs/AutomationTool/Scripts/BuildDocumentation.Automation.cs
--- a/Engine/Source/Programs/AutomationTool/Scripts/BuildDocumentation.Automation.cs
+++ b/Engine/Source/Programs/AutomationTool/Scripts/BuildDocumentation.Automation.cs
@@ -37,6 +37,11 @@ public class BuildDocumentation : BuildCommand
 		bool bBuildHtml = bBuild || ParseParam("buildhtml");
 		bool bBuildChm = bBuild || ParseParam("buildchm");
 
+		if(ParseParam("nodoxygen"))
+		{
+			bCleanMeta = bCleanEnv = bCleanXml = bBuildMeta = bBuildEnv = bBuildXml = false;
+		}
+
 		bool bStats = ParseParam("stats");
 
 		bool bMakeArchives = ParseParam("archive");
@@ -45,7 +50,6 @@ public class BuildDocumentation : BuildCommand
 
 		// Create the intermediate folders
 		string IntermediateDir = Path.Combine(CmdEnv.LocalRoot, "Engine\\Intermediate\\Documentation");
-		CommandUtils.DeleteDirectory(IntermediateDir);
 		CommandUtils.CreateDirectory(IntermediateDir);
 		string BuildDir = Path.Combine(IntermediateDir, "build");
 		CommandUtils.CreateDirectory(BuildDir);
diff --git a/Engine/Source/Programs/BlankProgram/Private/BlankProgram.cpp b/Engine/Source/Programs/BlankProgram/Private/BlankProgram.cpp
--- a/Engine/Source/Programs/BlankProgram/Private/BlankProgram.cpp
+++ b/Engine/Source/Programs/BlankProgram/Private/BlankProgram.cpp
@@ -8,7 +8,8 @@
 DEFINE_LOG_CATEGORY_STATIC(LogBlankProgram, Log, All);
 
 IMPLEMENT_APPLICATION(BlankProgram, "BlankProgram");
-int32 main(int32 ArgC, char* ArgV[])
+
+INT32_MAIN_INT32_ARGC_TCHAR_ARGV()
 {
 	GEngineLoop.PreInit(ArgC, ArgV);
 	UE_LOG(LogBlankProgram, Display, TEXT("Hello World"));
diff --git a/Engine/Source/Programs/CrashReporter/MinidumpDiagnostics/Private/MinidumpDiagnosticsApp.cpp b/Engine/Source/Programs/CrashReporter/MinidumpDiagnostics/Private/MinidumpDiagnosticsApp.cpp
--- a/Engine/Source/Programs/CrashReporter/MinidumpDiagnostics/Private/MinidumpDiagnosticsApp.cpp
+++ b/Engine/Source/Programs/CrashReporter/MinidumpDiagnostics/Private/MinidumpDiagnosticsApp.cpp
@@ -4,7 +4,7 @@
 #include "ModuleManager.h"
 #include "CrashDebugHelperModule.h"
 
-int32 RunMinidumpDiagnostics( int32 ArgC, ANSICHAR* Argv[] )
+int32 RunMinidumpDiagnostics(int32 ArgC, TCHAR* Argv[])
 {
 	// Make sure we have at least a single parameter
 	if( ArgC < 2 )
@@ -32,7 +32,7 @@ int32 RunMinidumpDiagnostics( int32 ArgC, ANSICHAR* Argv[] )
 	IModuleInterface& PerforceSourceControlModule = FModuleManager::LoadModuleChecked<IModuleInterface>( FName( "PerforceSourceControl" ) );
 
 	// Create a report for the minidump passed in on the command line
-	FString MinidumpName = ANSI_TO_TCHAR( Argv[1] );
+	FString MinidumpName = Argv[1];
 	CrashDebugHelper->CreateMinidumpDiagnosticReport( MinidumpName );
 
 	// Write a report next to the original minidump
@@ -48,3 +48,4 @@ int32 RunMinidumpDiagnostics( int32 ArgC, ANSICHAR* Argv[] )
 	return 0;
 }
 
+
diff --git a/Engine/Source/Programs/CrashReporter/MinidumpDiagnostics/Private/Windows/MinidumpDiagnosticsMainWindows.cpp b/Engine/Source/Programs/CrashReporter/MinidumpDiagnostics/Private/Windows/MinidumpDiagnosticsMainWindows.cpp
--- a/Engine/Source/Programs/CrashReporter/MinidumpDiagnostics/Private/Windows/MinidumpDiagnosticsMainWindows.cpp
+++ b/Engine/Source/Programs/CrashReporter/MinidumpDiagnostics/Private/Windows/MinidumpDiagnosticsMainWindows.cpp
@@ -9,7 +9,7 @@
 IMPLEMENT_APPLICATION( MinidumpDiagnostics, "MinidumpDiagnostics" )
 
 // More Windows glue
-int32 GuardedMain( int32 Argc, ANSICHAR* Argv[] )
+int32 GuardedMain(int32 Argc, TCHAR* Argv[])
 {
 	GEngineLoop.PreInit( Argc, Argv );
 
@@ -22,7 +22,7 @@ int32 GuardedMain( int32 Argc, ANSICHAR* Argv[] )
 }
 
 // Windows glue
-int32 GuardedMainWrapper( int32 ArgC, ANSICHAR* ArgV[] )
+int32 GuardedMainWrapper(int32 ArgC, TCHAR* ArgV[])
 {
 	int32 ReturnCode = 0;
 
@@ -54,7 +54,7 @@ int32 GuardedMainWrapper( int32 ArgC, ANSICHAR* ArgV[] )
 }
 
  // Main entry point to the application
-int32 main( int32 ArgC, ANSICHAR* ArgV[] )
+INT32_MAIN_INT32_ARGC_TCHAR_ARGV()
 {
 	const int32 Result = GuardedMainWrapper( ArgC, ArgV );
 	return Result;
diff --git a/Engine/Source/Programs/CrashReporter/MinidumpDiagnostics/Public/MinidumpDiagnosticsApp.h b/Engine/Source/Programs/CrashReporter/MinidumpDiagnostics/Public/MinidumpDiagnosticsApp.h
--- a/Engine/Source/Programs/CrashReporter/MinidumpDiagnostics/Public/MinidumpDiagnosticsApp.h
+++ b/Engine/Source/Programs/CrashReporter/MinidumpDiagnostics/Public/MinidumpDiagnosticsApp.h
@@ -5,5 +5,6 @@
 #include "Core.h"
 
 // Extract a debugging report from a minidump
-int32 RunMinidumpDiagnostics( int32 ArgC, ANSICHAR* Argv[] );
+int32 RunMinidumpDiagnostics(int32 ArgC, TCHAR* Argv[]);
+
 
diff --git a/Engine/Source/Programs/IOS/iPhonePackager/CookTime.cs b/Engine/Source/Programs/IOS/iPhonePackager/CookTime.cs
--- a/Engine/Source/Programs/IOS/iPhonePackager/CookTime.cs
+++ b/Engine/Source/Programs/IOS/iPhonePackager/CookTime.cs
@@ -75,6 +75,10 @@ namespace iPhonePackager
 				Program.ReturnCode = 129;
 			}
 
+            // Set encoding to support unicode filenames
+            Stub.AlternateEncodingUsage = ZipOption.Always;
+            Stub.AlternateEncoding = Encoding.UTF8;
+
 			return Stub;
 		}
 
diff --git a/Engine/Source/Programs/ShaderCompileWorker/Private/ShaderCompileWorker.cpp b/Engine/Source/Programs/ShaderCompileWorker/Private/ShaderCompileWorker.cpp
--- a/Engine/Source/Programs/ShaderCompileWorker/Private/ShaderCompileWorker.cpp
+++ b/Engine/Source/Programs/ShaderCompileWorker/Private/ShaderCompileWorker.cpp
@@ -486,7 +486,7 @@ private:
  *		The parent process Id
  *		The thread Id corresponding to this worker
  */
-int32 GuardedMain(int32 argc, ANSICHAR* argv[])
+int32 GuardedMain(int32 argc, TCHAR* argv[])
 {
 	GEngineLoop.PreInit(argc, argv, TEXT("-NOPACKAGECACHE -Multiprocess"));
 #if DEBUG_USING_CONSOLE
@@ -495,7 +495,7 @@ int32 GuardedMain(int32 argc, ANSICHAR* argv[])
 
 #if PLATFORM_WINDOWS
 	//@todo - would be nice to change application name or description to have the ThreadId in it for debugging purposes
-	SetConsoleTitle(ANSI_TO_TCHAR(argv[3]));
+	SetConsoleTitle(argv[3]);
 #endif
 
 	// We just enumerate the shader formats here for debugging.
@@ -514,7 +514,7 @@ int32 GuardedMain(int32 argc, ANSICHAR* argv[])
 
 	LastCompileTime = FPlatformTime::Seconds();
 
-	FString InCommunicating = ANSI_TO_TCHAR(argv[6]);
+	FString InCommunicating = argv[6];
 #if PLATFORM_SUPPORTS_NAMED_PIPES
 	const bool bThroughFile = (InCommunicating == FString(TEXT("-communicatethroughfile")));
 	const bool bThroughNamedPipe = (InCommunicating == FString(TEXT("-communicatethroughnamedpipe")));
@@ -527,14 +527,14 @@ int32 GuardedMain(int32 argc, ANSICHAR* argv[])
 	check((int32)bThroughFile + (int32)bThroughNamedPipe + (int32)bThroughNamedPipeOnce == 1);
 
 	FWorkLoop::ECommunicationMode Mode = bThroughFile ? FWorkLoop::ThroughFile : (bThroughNamedPipeOnce ? FWorkLoop::ThroughNamedPipeOnce : FWorkLoop::ThroughNamedPipe);
-	FWorkLoop WorkLoop(ANSI_TO_TCHAR(argv[2]), ANSI_TO_TCHAR(argv[1]), ANSI_TO_TCHAR(argv[4]), ANSI_TO_TCHAR(argv[5]), Mode);
+	FWorkLoop WorkLoop(argv[2], argv[1], argv[4], argv[5], Mode);
 
 	WorkLoop.Loop();
 
 	return 0;
 }
 
-int32 GuardedMainWrapper(int32 ArgC, ANSICHAR* ArgV[], const TCHAR* CrashOutputFile)
+int32 GuardedMainWrapper(int32 ArgC, TCHAR* ArgV[], const TCHAR* CrashOutputFile)
 {
 	int32 ReturnCode = 0;
 #if !PLATFORM_MAC
@@ -596,7 +596,8 @@ IMPLEMENT_APPLICATION(ShaderCompileWorker, "ShaderCompileWorker")
  * @param	ArgC	Command-line argument count
  * @param	ArgV	Argument strings
  */
-int32 main( int32 ArgC, ANSICHAR* ArgV[] )
+
+INT32_MAIN_INT32_ARGC_TCHAR_ARGV()
 {
 	if(ArgC < 6)
 	{
@@ -607,10 +608,9 @@ int32 main( int32 ArgC, ANSICHAR* ArgV[] )
 	// so just make sure we have at least the minimum number of parameters.
 	check(ArgC >= 6);
 
-	// ANSI_TO_TCHAR makes this need to be declared outside of a function with __try
 	TCHAR OutputFilePath[PLATFORM_MAX_FILEPATH_LENGTH];
-	FCString::Strncpy(OutputFilePath, ANSI_TO_TCHAR(ArgV[1]), PLATFORM_MAX_FILEPATH_LENGTH);
-	FCString::Strncat(OutputFilePath, ANSI_TO_TCHAR(ArgV[5]), PLATFORM_MAX_FILEPATH_LENGTH);
+	FCString::Strncpy(OutputFilePath, ArgV[1], PLATFORM_MAX_FILEPATH_LENGTH);
+	FCString::Strncat(OutputFilePath, ArgV[5], PLATFORM_MAX_FILEPATH_LENGTH);
 
 	const int32 ReturnCode = GuardedMainWrapper(ArgC,ArgV,OutputFilePath);
 	return ReturnCode;
diff --git a/Engine/Source/Programs/UnrealFileServer/Private/UnrealFileServer.cpp b/Engine/Source/Programs/UnrealFileServer/Private/UnrealFileServer.cpp
--- a/Engine/Source/Programs/UnrealFileServer/Private/UnrealFileServer.cpp
+++ b/Engine/Source/Programs/UnrealFileServer/Private/UnrealFileServer.cpp
@@ -17,7 +17,7 @@ IMPLEMENT_APPLICATION(UnrealFileServer, "UnrealFileServer");
  * @param	ArgC	Command-line argument count
  * @param	ArgV	Argument strings
  */
-int32 main(int32 ArgC, char* ArgV[])
+INT32_MAIN_INT32_ARGC_TCHAR_ARGV()
 {
 	// start up the main loop
 	GEngineLoop.PreInit(ArgC, ArgV);
diff --git a/Engine/Source/Programs/UnrealHeaderTool/Private/UnrealHeaderToolMain.cpp b/Engine/Source/Programs/UnrealHeaderTool/Private/UnrealHeaderToolMain.cpp
--- a/Engine/Source/Programs/UnrealHeaderTool/Private/UnrealHeaderToolMain.cpp
+++ b/Engine/Source/Programs/UnrealHeaderTool/Private/UnrealHeaderToolMain.cpp
@@ -15,7 +15,7 @@ DEFINE_LOG_CATEGORY(LogCompile);
  * @param	ArgC	Command-line argument count
  * @param	ArgV	Argument strings
  */
-int32 main( int32 ArgC, char* ArgV[] )
+INT32_MAIN_INT32_ARGC_TCHAR_ARGV()
 {
 	FString CmdLine;
 
diff --git a/Engine/Source/Programs/UnrealPak/Private/UnrealPak.cpp b/Engine/Source/Programs/UnrealPak/Private/UnrealPak.cpp
--- a/Engine/Source/Programs/UnrealPak/Private/UnrealPak.cpp
+++ b/Engine/Source/Programs/UnrealPak/Private/UnrealPak.cpp
@@ -132,7 +132,7 @@ bool CopyFileToPak(FArchive& InPak, const FString& InMountPoint, const FPakInput
 	return bFileExists;
 }
 
-void ProcessOrderFile(int32 ArgC, ANSICHAR* ArgV[], TMap<FString, uint64>& OrderMap)
+void ProcessOrderFile(int32 ArgC, TCHAR* ArgV[], TMap<FString, uint64>& OrderMap)
 {
 	// List of all items to add to pak file
 	FString ResponseFile;
@@ -174,7 +174,7 @@ void ProcessOrderFile(int32 ArgC, ANSICHAR* ArgV[], TMap<FString, uint64>& Order
 	}
 }
 
-void ProcessCommandLine(int32 ArgC, ANSICHAR* ArgV[], TArray<FPakInputPair>& Entries)
+void ProcessCommandLine(int32 ArgC, TCHAR* ArgV[], TArray<FPakInputPair>& Entries)
 {
 	// List of all items to add to pak file
 	FString ResponseFile;
@@ -241,7 +241,7 @@ void ProcessCommandLine(int32 ArgC, ANSICHAR* ArgV[], TArray<FPakInputPair>& Ent
 		for (int32 Index = 2; Index < ArgC; Index++)
 		{
 			// Skip switches and add everything else to the Entries array
-			ANSICHAR* Param = ArgV[Index];
+			TCHAR* Param = ArgV[Index];
 			if (Param[0] != '-')
 			{
 				FPakInputPair Input;
@@ -596,7 +596,7 @@ bool ExtractFilesFromPak(const TCHAR* InPakFilename, const TCHAR* InDestPath)
  * @param	ArgC	Command-line argument count
  * @param	ArgV	Argument strings
  */
-int32 main(int32 ArgC, ANSICHAR* ArgV[])
+INT32_MAIN_INT32_ARGC_TCHAR_ARGV()
 {
 	// start up the main loop
 	GEngineLoop.PreInit(ArgC, ArgV);
diff --git a/Engine/Source/Runtime/Core/Private/Apple/ApplePlatformFile.cpp b/Engine/Source/Runtime/Core/Private/Apple/ApplePlatformFile.cpp
--- a/Engine/Source/Runtime/Core/Private/Apple/ApplePlatformFile.cpp
+++ b/Engine/Source/Runtime/Core/Private/Apple/ApplePlatformFile.cpp
@@ -499,9 +499,11 @@ bool FApplePlatformFile::IterateDirectory(const TCHAR* Directory, FDirectoryVisi
 			struct dirent *Entry;
 			while ((Entry = readdir(Handle)) != NULL)
 			{
-				if (FCString::Strcmp(UTF8_TO_TCHAR(Entry->d_name), TEXT(".")) && FCString::Strcmp(UTF8_TO_TCHAR(Entry->d_name), TEXT("..")) && FCString::Strcmp(UTF8_TO_TCHAR(Entry->d_name), TEXT(".DS_Store")))
+				if (FCStringAnsi::Strcmp(Entry->d_name, ".") && FCStringAnsi::Strcmp(Entry->d_name, "..") && FCStringAnsi::Strcmp(Entry->d_name, ".DS_Store"))
 				{
-					Result = Visitor.Visit(*(FString(Directory) / UTF8_TO_TCHAR(Entry->d_name)), Entry->d_type == DT_DIR);
+                    // Normalize any unicode forms so we match correctly
+                    FString NormalizedFilename = UTF8_TO_TCHAR(([[[NSString stringWithUTF8String:Entry->d_name] precomposedStringWithCanonicalMapping] cStringUsingEncoding:NSUTF8StringEncoding]));
+                    Result = Visitor.Visit(*(FString(Directory) / NormalizedFilename), Entry->d_type == DT_DIR);
 				}
 			}
 			closedir(Handle);
diff --git a/Engine/Source/Runtime/Core/Private/GenericPlatform/GenericPlatformMisc.cpp b/Engine/Source/Runtime/Core/Private/GenericPlatform/GenericPlatformMisc.cpp
--- a/Engine/Source/Runtime/Core/Private/GenericPlatform/GenericPlatformMisc.cpp
+++ b/Engine/Source/Runtime/Core/Private/GenericPlatform/GenericPlatformMisc.cpp
@@ -416,7 +416,8 @@ void GenericPlatformMisc_GetProjectFilePathGameDir(FString& OutGameDir)
 	FString BasePath = FPaths::GetPath(FPaths::GetProjectFilePath());
 	FPaths::NormalizeFilename(BasePath);
 	BasePath = FFileManagerGeneric::DefaultConvertToRelativePath(*BasePath);
-	OutGameDir = FString::Printf(TEXT("%s/"), *BasePath);
+	if(!BasePath.EndsWith("/")) BasePath += TEXT("/");
+	OutGameDir = BasePath;
 }
 
 const TCHAR* FGenericPlatformMisc::GameDir()
@@ -521,6 +522,7 @@ const TCHAR* FGenericPlatformMisc::GameDir()
 				GameDir = FPaths::EngineUserDir();
 				FPaths::NormalizeFilename(GameDir);
 				GameDir = FFileManagerGeneric::DefaultConvertToRelativePath(*GameDir);
+				if(!GameDir.EndsWith(TEXT("/"))) GameDir += TEXT("/");
 			}
 		}
 	}
diff --git a/Engine/Source/Runtime/Core/Private/GenericPlatform/StandardPlatformString.cpp b/Engine/Source/Runtime/Core/Private/GenericPlatform/StandardPlatformString.cpp
--- a/Engine/Source/Runtime/Core/Private/GenericPlatform/StandardPlatformString.cpp
+++ b/Engine/Source/Runtime/Core/Private/GenericPlatform/StandardPlatformString.cpp
@@ -10,6 +10,12 @@ DEFINE_LOG_CATEGORY_STATIC(LogStandardPlatformString, Log, All);
 
 #if !PLATFORM_USE_SYSTEM_VSWPRINTF
 
+#if PLATFORM_IOS
+#define VA_LIST_REF va_list&
+#else
+#define VA_LIST_REF va_list
+#endif
+
 struct FFormatInfo
 {
 	WIDECHAR Format[32];
@@ -32,7 +38,7 @@ static int32 GetFormattingInfo(const WIDECHAR* Format, FFormatInfo& OutInfo)
 	OutInfo.HasDynamicWidth = false;
 
 	// Skip width
-	while (*Format >= LITERAL(WIDECHAR, '0') && *Format <= LITERAL(WIDECHAR, '9') || *Format == LITERAL(WIDECHAR, '*'))
+	while ((*Format >= LITERAL(WIDECHAR, '0') && *Format <= LITERAL(WIDECHAR, '9')) || *Format == LITERAL(WIDECHAR, '*'))
 	{
 		if (*Format == LITERAL(WIDECHAR, '*'))
 		{
@@ -45,7 +51,7 @@ static int32 GetFormattingInfo(const WIDECHAR* Format, FFormatInfo& OutInfo)
 	if (*Format == LITERAL(WIDECHAR, '.'))
 	{
 		Format++;
-		while (*Format >= LITERAL(WIDECHAR, '0') && *Format <= LITERAL(WIDECHAR, '9') || *Format == LITERAL(WIDECHAR, '*'))
+		while ((*Format >= LITERAL(WIDECHAR, '0') && *Format <= LITERAL(WIDECHAR, '9')) || *Format == LITERAL(WIDECHAR, '*'))
 		{
 			if (*Format == LITERAL(WIDECHAR, '*'))
 			{
@@ -83,7 +89,7 @@ static int32 GetFormattingInfo(const WIDECHAR* Format, FFormatInfo& OutInfo)
 }
 
 template <typename T1, typename T2>
-static int32 FormatString(const FFormatInfo& Info, va_list ArgPtr, WIDECHAR* Formatted, int32 Length)
+static int32 FormatString(const FFormatInfo& Info, VA_LIST_REF ArgPtr, WIDECHAR* Formatted, int32 Length)
 {
 	if (Info.HasDynamicWidth)
 	{
@@ -114,7 +120,7 @@ static int32 FormatString(const FFormatInfo& Info, va_list ArgPtr, WIDECHAR* For
 	}
 }
 
-static const WIDECHAR* GetFormattedArgument(const FFormatInfo& Info, va_list ArgPtr, WIDECHAR* Formatted, int32 &InOutLength)
+static const WIDECHAR* GetFormattedArgument(const FFormatInfo& Info, VA_LIST_REF ArgPtr, WIDECHAR* Formatted, int32 &InOutLength)
 {
 	if (FChar::ToLower(Info.Type) == LITERAL(WIDECHAR, 's'))
 	{
@@ -158,7 +164,7 @@ static const WIDECHAR* GetFormattedArgument(const FFormatInfo& Info, va_list Arg
 
 int32 FStandardPlatformString::GetVarArgs( WIDECHAR* Dest, SIZE_T DestSize, int32 Count, const WIDECHAR*& Fmt, va_list ArgPtr )
 {
-	const WIDECHAR* Format = Fmt;
+  	const WIDECHAR* Format = Fmt;
 	const WIDECHAR* DestStart = Dest;
 
 	while (Count > 1 &&  *Format)
diff --git a/Engine/Source/Runtime/Core/Private/IOS/IOSPlatformFile.cpp b/Engine/Source/Runtime/Core/Private/IOS/IOSPlatformFile.cpp
--- a/Engine/Source/Runtime/Core/Private/IOS/IOSPlatformFile.cpp
+++ b/Engine/Source/Runtime/Core/Private/IOS/IOSPlatformFile.cpp
@@ -370,9 +370,11 @@ bool FIOSPlatformFile::IterateDirectory(const TCHAR* Directory, FDirectoryVisito
 		struct dirent *Entry;
 		while ((Entry = readdir(Handle)) != NULL)
 		{
-			if (FCString::Strcmp(UTF8_TO_TCHAR(Entry->d_name), TEXT(".")) && FCString::Strcmp(UTF8_TO_TCHAR(Entry->d_name), TEXT("..")))
+			if (FCStringAnsi::Strcmp(Entry->d_name, ".") && FCStringAnsi::Strcmp(Entry->d_name, ".."))
 			{
-				Result = Visitor.Visit(*(FString(Directory) / UTF8_TO_TCHAR(Entry->d_name)), Entry->d_type == DT_DIR);
+				// Normalize any unicode forms so we match correctly
+				FString NormalizedFilename = UTF8_TO_TCHAR(([[[NSString stringWithUTF8String:Entry->d_name] precomposedStringWithCanonicalMapping] cStringUsingEncoding:NSUTF8StringEncoding]));
+				Result = Visitor.Visit(*(FString(Directory) / NormalizedFilename), Entry->d_type == DT_DIR);
 			}
 		}
 		closedir(Handle);
diff --git a/Engine/Source/Runtime/Core/Private/IOS/IOSPlatformMisc.cpp b/Engine/Source/Runtime/Core/Private/IOS/IOSPlatformMisc.cpp
--- a/Engine/Source/Runtime/Core/Private/IOS/IOSPlatformMisc.cpp
+++ b/Engine/Source/Runtime/Core/Private/IOS/IOSPlatformMisc.cpp
@@ -134,7 +134,7 @@ void FIOSPlatformMisc::GetEnvironmentVariable(const TCHAR* VariableName, TCHAR*
 void FIOSPlatformMisc::LowLevelOutputDebugString( const TCHAR *Message )
 {
 	//NsLog will out to all iOS output consoles, instead of just the Xcode console.
-	NSLog(@"%s", TCHAR_TO_ANSI(Message));
+	NSLog(@"%@", [NSString stringWithUTF8String:TCHAR_TO_UTF8(Message)]);
 }
 
 const TCHAR* FIOSPlatformMisc::GetSystemErrorMessage(TCHAR* OutBuffer, int32 BufferCount, int32 Error)
@@ -531,3 +531,4 @@ FString FIOSPlatformMisc::GetUniqueDeviceId()
 	}
 	return FPlatformMisc::GetHashedMacAddressString();
 }
+
diff --git a/Engine/Source/Runtime/Core/Public/Apple/ApplePlatformRunnableThread.h b/Engine/Source/Runtime/Core/Public/Apple/ApplePlatformRunnableThread.h
--- a/Engine/Source/Runtime/Core/Public/Apple/ApplePlatformRunnableThread.h
+++ b/Engine/Source/Runtime/Core/Public/Apple/ApplePlatformRunnableThread.h
@@ -38,6 +38,22 @@ private:
 		FPlatformMisc::ReleaseAutoreleasePool(Pool);
 	}
     
+	/**
+	 * Allows platforms to adjust stack size
+	 */
+	virtual uint32 AdjustStackSize(uint32 InStackSize)
+	{
+		InStackSize = FRunnableThreadPThread::AdjustStackSize(InStackSize);
+        
+		// If it's set, make sure it's at least 128 KB or stack allocations (e.g. in Logf) may fail
+		if (InStackSize && InStackSize < 128 * 1024)
+		{
+			InStackSize = 128 * 1024;
+		}
+        
+		return InStackSize;
+	}
+    
     void*      Pool;
 };
 
diff --git a/Engine/Source/Runtime/Core/Public/Core.h b/Engine/Source/Runtime/Core/Public/Core.h
--- a/Engine/Source/Runtime/Core/Public/Core.h
+++ b/Engine/Source/Runtime/Core/Public/Core.h
@@ -480,24 +480,35 @@ extern CORE_API double GBlueprintCompileTime;
 #endif
 
 
-// Trick to prevent people from using old UE4 types (which are still defined in Windef.h on PC).
-// If old type is really necessary please use the global scope operator to access it (i.e. ::INT).
-// Windows header file includes can be wrapped in #include "AllowWindowsPlatformTypes.h" and
-// #include "HideWindowsPlatformTypes.h"
+/// Trick to prevent people from using old UE4 types (which are still defined in Windef.h on PC).
 namespace DoNotUseOldUE4Type
 {
+	/// Used to cause compile errors through typedefs of unportable types. If unportable types 
+	/// are really necessary please use the global scope operator to access it (i.e. ::&nbsp;INT). 
+	/// Windows header file includes can be wrapped in #include "AllowWindowsPlatformTypes.h" 
+	/// and #include "HideWindowsPlatformTypes.h"
 	class FUnusableType
 	{
 		FUnusableType();
 		~FUnusableType();
 	};
 
+	/// This type is aliased to FUnusableType - use ::int32 instead.
 	typedef FUnusableType INT;
+
+	/// This type is aliased to FUnusableType - use ::uint32 instead.
 	typedef FUnusableType UINT;
+
+	/// This type is aliased to FUnusableType - use ::uint32 instead.
 	typedef FUnusableType DWORD;
+
+	/// This type is aliased to FUnusableType - use float instead.
 	typedef FUnusableType FLOAT;
 
+	/// This value is aliased to FUnusableType - use true instead.
 	typedef FUnusableType TRUE;
+
+	/// This value is aliased to FUnusableType - use false instead.
 	typedef FUnusableType FALSE;
 }
 
diff --git a/Engine/Source/Runtime/Core/Public/HAL/Platform.h b/Engine/Source/Runtime/Core/Public/HAL/Platform.h
--- a/Engine/Source/Runtime/Core/Public/HAL/Platform.h
+++ b/Engine/Source/Runtime/Core/Public/HAL/Platform.h
@@ -137,6 +137,9 @@
 #ifndef PLATFORM_COMPILER_COMMON_LANGUAGE_RUNTIME_COMPILATION
 	#define PLATFORM_COMPILER_COMMON_LANGUAGE_RUNTIME_COMPILATION 0
 #endif
+#ifndef PLATFORM_COMPILER_HAS_TCHAR_WMAIN
+	#define PLATFORM_COMPILER_HAS_TCHAR_WMAIN 0
+#endif
 #ifndef PLATFORM_TCHAR_IS_4_BYTES
 	#define PLATFORM_TCHAR_IS_4_BYTES			0
 #endif
@@ -445,6 +448,33 @@ struct THasOperatorNotEquals
 			void operator !=( const __VA_ARGS__ & RHS ) const;
 #endif
 
+
+// Console ANSICHAR/TCHAR command line handling
+#if PLATFORM_COMPILER_HAS_TCHAR_WMAIN
+#define INT32_MAIN_INT32_ARGC_TCHAR_ARGV() int32 wmain(int32 ArgC, TCHAR* ArgV[])
+#else
+#define INT32_MAIN_INT32_ARGC_TCHAR_ARGV() \
+int32 tchar_main(int32 ArgC, TCHAR* ArgV[]); \
+int32 main(int32 ArgC, ANSICHAR* Utf8ArgV[]) \
+{ \
+	TCHAR** ArgV = new TCHAR*[ArgC]; \
+	for (int32 a = 0; a < ArgC; a++) \
+	{ \
+		FUTF8ToTCHAR ConvertFromUtf8(Utf8ArgV[a]); \
+		ArgV[a] = new TCHAR[ConvertFromUtf8.Length() + 1]; \
+		FCString::Strcpy(ArgV[a], ConvertFromUtf8.Length(), ConvertFromUtf8.Get()); \
+	} \
+	int32 Result = tchar_main(ArgC, ArgV); \
+	for (int32 a = 0; a < ArgC; a++) \
+	{ \
+		delete[] ArgV[a]; \
+	} \
+	delete[] ArgV; \
+	return Result; \
+} \
+int32 tchar_main(int32 ArgC, TCHAR* ArgV[])
+#endif
+
 //--------------------------------------------------------------------------------------------
 // POD types refactor for porting old code:
 //
diff --git a/Engine/Source/Runtime/Core/Public/IOS/IOSCommandLineHelper.h b/Engine/Source/Runtime/Core/Public/IOS/IOSCommandLineHelper.h
--- a/Engine/Source/Runtime/Core/Public/IOS/IOSCommandLineHelper.h
+++ b/Engine/Source/Runtime/Core/Public/IOS/IOSCommandLineHelper.h
@@ -70,7 +70,7 @@ class FIOSCommandLineHelper
 					CommandLine[strlen(CommandLine) - 1] = 0;
 				}
 
-				FCommandLine::Append(ANSI_TO_TCHAR(CommandLine));
+				FCommandLine::Append(UTF8_TO_TCHAR(CommandLine));
 			}
 			else
 			{
diff --git a/Engine/Source/Runtime/Core/Public/IOS/IOSPlatform.h b/Engine/Source/Runtime/Core/Public/IOS/IOSPlatform.h
--- a/Engine/Source/Runtime/Core/Public/IOS/IOSPlatform.h
+++ b/Engine/Source/Runtime/Core/Public/IOS/IOSPlatform.h
@@ -34,6 +34,7 @@ typedef FIOSPlatformTypes FPlatformTypes;
 #define PLATFORM_COMPILER_DISTINGUISHES_INT_AND_LONG	1
 #define PLATFORM_TCHAR_IS_4_BYTES						1
 #define PLATFORM_HAS_vsnprintf							0
+#define PLATFORM_USE_SYSTEM_VSWPRINTF					0
 #define PLATFORM_HAS_BSD_TIME							1
 #define PLATFORM_MAX_FILEPATH_LENGTH					MAX_PATH
 #define PLATFORM_SUPPORTS_TEXTURE_STREAMING				1
diff --git a/Engine/Source/Runtime/Core/Public/Windows/WIndowsPlatform.h b/Engine/Source/Runtime/Core/Public/Windows/WIndowsPlatform.h
--- a/Engine/Source/Runtime/Core/Public/Windows/WIndowsPlatform.h
+++ b/Engine/Source/Runtime/Core/Public/Windows/WIndowsPlatform.h
@@ -69,6 +69,7 @@ typedef FWindowsPlatformTypes FPlatformTypes;
 #define PLATFORM_COMPILER_HAS_DEFAULTED_FUNCTIONS	0
 #define PLATFORM_COMPILER_HAS_VARIADIC_TEMPLATES	0
 #define PLATFORM_COMPILER_HAS_EXPLICIT_OPERATORS	0
+#define PLATFORM_COMPILER_HAS_TCHAR_WMAIN			1
 
 #define PLATFORM_HAS_128BIT_ATOMICS					PLATFORM_64BITS
 
diff --git a/Engine/Source/Runtime/Launch/Private/IOS/LaunchIOS.cpp b/Engine/Source/Runtime/Launch/Private/IOS/LaunchIOS.cpp
--- a/Engine/Source/Runtime/Launch/Private/IOS/LaunchIOS.cpp
+++ b/Engine/Source/Runtime/Launch/Private/IOS/LaunchIOS.cpp
@@ -171,14 +171,27 @@ FString GSavedCommandLine;
 
 int main(int argc, char *argv[])
 {
-	for(int Option = 1; Option < argc; Option++)
+    for(int Option = 1; Option < argc; Option++)
 	{
 		GSavedCommandLine += TEXT(" ");
-		GSavedCommandLine += ANSI_TO_TCHAR(argv[Option]);
+		GSavedCommandLine += UTF8_TO_TCHAR(argv[Option]);
 	}
 
 	FIOSCommandLineHelper::InitCommandArgs(FString());
 	
+#if UE_BUILD_DEBUG
+    if (FParse::Param(FCommandLine::Get(), TEXT("WaitForDebugger")))
+    {
+        while(!FPlatformMisc::IsDebuggerPresent())
+        {
+            FPlatformMisc::LowLevelOutputDebugString(TEXT("Waiting for debugger...\n"));
+            FPlatformProcess::Sleep(1.f);
+        }
+        FPlatformMisc::LowLevelOutputDebugString(TEXT("Debugger attached.\n"));
+    }
+#endif
+    
+    
 	@autoreleasepool {
 	    return UIApplicationMain(argc, argv, nil, NSStringFromClass([IOSAppDelegate class]));
 	}
diff --git a/Engine/Source/Runtime/Launch/Private/LaunchEngineLoop.cpp b/Engine/Source/Runtime/Launch/Private/LaunchEngineLoop.cpp
--- a/Engine/Source/Runtime/Launch/Private/LaunchEngineLoop.cpp
+++ b/Engine/Source/Runtime/Launch/Private/LaunchEngineLoop.cpp
@@ -476,7 +476,7 @@ FEngineLoop::FEngineLoop()
 { }
 
 
-int32 FEngineLoop::PreInit(int32 ArgC, ANSICHAR* ArgV[], const TCHAR* AdditionalCommandline)
+int32 FEngineLoop::PreInit(int32 ArgC, TCHAR* ArgV[], const TCHAR* AdditionalCommandline)
 {
 	FString CmdLine;
 
@@ -2348,6 +2348,7 @@ void FEngineLoop::AppInit( )
 	UE_LOG(LogInit, Log, TEXT("Command line: %s"), FCommandLine::Get() );
 	UE_LOG(LogInit, Log, TEXT("Base directory: %s"), FPlatformProcess::BaseDir() );
 	//UE_LOG(LogInit, Log, TEXT("Character set: %s"), sizeof(TCHAR)==1 ? TEXT("ANSI") : TEXT("Unicode") );
+	UE_LOG(LogInit, Log, TEXT("Rocket: %d"), FRocketSupport::IsRocket()? 1 : 0);
 
 	GPrintLogTimes = ELogTimes::None;
 
diff --git a/Engine/Source/Runtime/Launch/Private/Mac/LaunchMac.cpp b/Engine/Source/Runtime/Launch/Private/Mac/LaunchMac.cpp
--- a/Engine/Source/Runtime/Launch/Private/Mac/LaunchMac.cpp
+++ b/Engine/Source/Runtime/Launch/Private/Mac/LaunchMac.cpp
@@ -235,12 +235,12 @@ void EngineCrashHandler(const FGenericCrashContext & GenericContext)
 
 @end
 
-int main(int argc, char *argv[])
+INT32_MAIN_INT32_ARGC_TCHAR_ARGV()
 {
-	for (int32 Option = 1; Option < argc; Option++)
+	for (int32 Option = 1; Option < ArgC; Option++)
 	{
 		GSavedCommandLine += TEXT(" ");
-		FString Argument(ANSI_TO_TCHAR(argv[Option]));
+		FString Argument(ArgV[Option]);
 		if (Argument.Contains(TEXT(" ")))
 		{
 			if (Argument.Contains(TEXT("=")))
@@ -258,5 +258,5 @@ int main(int argc, char *argv[])
 		GSavedCommandLine += Argument;
 	}
 
-	return NSApplicationMain(argc, (const char **)argv);
+	return NSApplicationMain(ArgC, (const char **)ArgV);
 }
diff --git a/Engine/Source/Runtime/Launch/Public/LaunchEngineLoop.h b/Engine/Source/Runtime/Launch/Public/LaunchEngineLoop.h
--- a/Engine/Source/Runtime/Launch/Public/LaunchEngineLoop.h
+++ b/Engine/Source/Runtime/Launch/Public/LaunchEngineLoop.h
@@ -33,7 +33,7 @@ public:
 	 *
 	 * @return Returns the error level, 0 if successful and > 0 if there were errors
 	 */ 
-	int32 PreInit(int32 ArgC, ANSICHAR* ArgV[], const TCHAR* AdditionalCommandline = NULL);
+	int32 PreInit(int32 ArgC, TCHAR* ArgV[], const TCHAR* AdditionalCommandline = NULL);
 
 	/**
 	 * Pre-Initialize the main loop - parse command line, sets up GIsEditor, etc.
