Fixed VS_OPEN logspam caused by a deferred command being constantly requested until VS had finished opening

TTP# 331398 - EDITOR: Logspam when opening Visual Studio

Deferring a VS file open request used to use a deferred command to wait until VS had finished opening before opening the requested files, however, this command was essentially simulating a Tick, as if the command was not yet ready to be processed it was simply deferred again, leading to spammed output.

Rather than do this, I've added a Tick function which will attempt to open the entire list of pending files as a single call, rather than try each file individually like the command would.

diff --git a/Engine/Plugins/Developer/VisualStudioSourceCodeAccess/Source/VisualStudioSourceCodeAccess/Private/VisualStudioSourceCodeAccessor.cpp b/Engine/Plugins/Developer/VisualStudioSourceCodeAccess/Source/VisualStudioSourceCodeAccess/Private/VisualStudioSourceCodeAccessor.cpp
--- a/Engine/Plugins/Developer/VisualStudioSourceCodeAccess/Source/VisualStudioSourceCodeAccess/Private/VisualStudioSourceCodeAccessor.cpp
+++ b/Engine/Plugins/Developer/VisualStudioSourceCodeAccess/Source/VisualStudioSourceCodeAccess/Private/VisualStudioSourceCodeAccessor.cpp
@@ -368,11 +368,8 @@ bool FVisualStudioSourceCodeAccessor::OpenVisualStudioFilesInternalViaDTE(const
 		// Defer the request until VS is available to take hold of
 		if ( bDefer )
 		{
-			for ( const FileOpenRequest& Request : Requests )
-			{
-				const FString DeferCommand = FString::Printf(TEXT("OPEN_VS %s %d %d"), *Request.FullPath, Request.LineNumber, Request.ColumnNumber);
-				SourceCodeAccessModule.OnLaunchCodeAccessorDeferred().Broadcast( DeferCommand );
-			}
+			FScopeLock Lock(&DeferredRequestsCriticalSection);
+			DeferredRequests.Append(Requests);
 		}
 		else if ( !bSuccess )
 		{
@@ -424,36 +421,6 @@ bool FVisualStudioSourceCodeAccessor::SaveAllOpenDocuments() const
 	return bSuccess;
 }
 
-bool FVisualStudioSourceCodeAccessor::Exec( UWorld* InWorld, const TCHAR* Cmd, FOutputDevice& Ar )
-{
-	const TCHAR *Str = Cmd;
-	if( FParse::Command(&Str,TEXT("OPEN_VS")) )
-	{
-		// Valid column number?
-		const FString Params( Str );
-		const FString Space( TEXT(" ") );
-		const int32 ColumnIndex = Params.Find( Space, ESearchCase::CaseSensitive, ESearchDir::FromEnd );
-		if ( ColumnIndex != INDEX_NONE )
-		{
-			// Valid line number?
-			const int32 ColumnNumber = FCString::Atoi( *Params.RightChop( ColumnIndex+Space.Len() ) );
-			const int32 LineIndex = Params.Find( Space, ESearchCase::CaseSensitive, ESearchDir::FromEnd, ColumnIndex );
-			if ( LineIndex != INDEX_NONE )
-			{
-				// Valid source path?
-				const int32 LineNumber = FCString::Atoi( *Params.Mid( LineIndex+Space.Len(), ColumnIndex-LineIndex-Space.Len() ) );
-				const FString FullPath = Params.LeftChop( Params.Len() - LineIndex );
-				if ( FullPath.Len() > 0 )
-				{
-					OpenVisualStudioFileAtLineInternal( FullPath, LineNumber, ColumnNumber );
-				}
-			}
-		}
-		return true;
-	}
-	return false;
-}
-
 void FVisualStudioSourceCodeAccessor::VSLaunchStarted()
 {
 	// Broadcast the info and hope that MainFrame is around to receive it
@@ -473,10 +440,9 @@ void FVisualStudioSourceCodeAccessor::VSLaunchFinished( bool bSuccess )
 #else
 
 // VS Express-only dummy versions
-bool FVSAccessorModule::SaveAllOpenVisualStudioDocuments() { return false; }
-bool FVSAccessorModule::Exec( UWorld* InWorld, const TCHAR* Cmd, FOutputDevice& Ar ) { return false; }
-void FVSAccessorModule::VSLaunchStarted() {}
-void FVSAccessorModule::VSLaunchFinished( bool bSuccess ) {}
+bool FVisualStudioSourceCodeAccessor::SaveAllOpenDocuments() const { return false; }
+void FVisualStudioSourceCodeAccessor::VSLaunchStarted() {}
+void FVisualStudioSourceCodeAccessor::VSLaunchFinished( bool bSuccess ) {}
 
 #endif
 
@@ -954,4 +920,22 @@ FString FVisualStudioSourceCodeAccessor::GetSolutionPath() const
 	return CachedSolutionPath;
 }
 
+void FVisualStudioSourceCodeAccessor::Tick(const float DeltaTime)
+{
+	if (IsVSLaunchInProgress())
+	{
+		TArray<FileOpenRequest> TmpDeferredRequests;
+		{
+			FScopeLock Lock(&DeferredRequestsCriticalSection);
+
+			// Copy the DeferredRequests array, as OpenVisualStudioFilesInternal may update it
+			TmpDeferredRequests = DeferredRequests;
+			DeferredRequests.Empty();
+		}
+
+		// Try and open any pending files in VS first (this will update the VS launch state appropriately)
+		OpenVisualStudioFilesInternal(TmpDeferredRequests);
+	}
+}
+
 #undef LOCTEXT_NAMESPACE
\ No newline at end of file
diff --git a/Engine/Plugins/Developer/VisualStudioSourceCodeAccess/Source/VisualStudioSourceCodeAccess/Private/VisualStudioSourceCodeAccessor.h b/Engine/Plugins/Developer/VisualStudioSourceCodeAccess/Source/VisualStudioSourceCodeAccess/Private/VisualStudioSourceCodeAccessor.h
--- a/Engine/Plugins/Developer/VisualStudioSourceCodeAccess/Source/VisualStudioSourceCodeAccess/Private/VisualStudioSourceCodeAccessor.h
+++ b/Engine/Plugins/Developer/VisualStudioSourceCodeAccess/Source/VisualStudioSourceCodeAccess/Private/VisualStudioSourceCodeAccessor.h
@@ -3,7 +3,7 @@
 
 #include "ISourceCodeAccessor.h"
 
-class FVisualStudioSourceCodeAccessor : public ISourceCodeAccessor, FSelfRegisteringExec
+class FVisualStudioSourceCodeAccessor : public ISourceCodeAccessor
 {
 public:
 	/** Struct representing identifying information about Visual Studio versions */
@@ -30,6 +30,7 @@ public:
 	virtual bool OpenFileAtLine(const FString& FullPath, int32 LineNumber, int32 ColumnNumber = 0) OVERRIDE;
 	virtual bool OpenSourceFiles(const TArray<FString>& AbsoluteSourcePaths) OVERRIDE;
 	virtual bool SaveAllOpenDocuments() const OVERRIDE;
+	virtual void Tick(const float DeltaTime) OVERRIDE;
 
 private:
 
@@ -53,10 +54,6 @@ private:
 		}
 	};
 
-	// Begin Exec Interface
-	virtual bool Exec( UWorld* InWorld, const TCHAR* Cmd, FOutputDevice& Ar ) OVERRIDE;
-	// End Exec Interface
-
 	/** Checks to see if we can run visual studio, also outputs the executable path if it's needed */
 	bool CanRunVisualStudio(FString& OutPath) const;
 
@@ -126,6 +123,12 @@ private:
 	/** If !0 it represents the time at which the a VS instance was opened */
 	double	VSLaunchTime;
 
+	/** Critical section for updating DeferredRequests */
+	FCriticalSection DeferredRequestsCriticalSection;
+
+	/** Array of file open requests that have been deferred because Visual Studio was not available when the request was made */
+	TArray<FileOpenRequest> DeferredRequests;
+
 	/** Accessor for SolutionPath. Will try to update it when called from the game thread, otherwise will use the cached value */
 	FString GetSolutionPath() const;
 };
diff --git a/Engine/Source/Developer/SourceCodeAccess/Private/DefaultSourceCodeAccessor.cpp b/Engine/Source/Developer/SourceCodeAccess/Private/DefaultSourceCodeAccessor.cpp
--- a/Engine/Source/Developer/SourceCodeAccess/Private/DefaultSourceCodeAccessor.cpp
+++ b/Engine/Source/Developer/SourceCodeAccess/Private/DefaultSourceCodeAccessor.cpp
@@ -45,4 +45,8 @@ bool FDefaultSourceCodeAccessor::SaveAllOpenDocuments() const
 	return false;
 }
 
+void FDefaultSourceCodeAccessor::Tick(const float DeltaTime)
+{
+}
+
 #undef LOCTEXT_NAMESPACE
\ No newline at end of file
diff --git a/Engine/Source/Developer/SourceCodeAccess/Private/DefaultSourceCodeAccessor.h b/Engine/Source/Developer/SourceCodeAccess/Private/DefaultSourceCodeAccessor.h
--- a/Engine/Source/Developer/SourceCodeAccess/Private/DefaultSourceCodeAccessor.h
+++ b/Engine/Source/Developer/SourceCodeAccess/Private/DefaultSourceCodeAccessor.h
@@ -16,4 +16,5 @@ public:
 	virtual bool OpenFileAtLine(const FString& FullPath, int32 LineNumber, int32 ColumnNumber = 0) OVERRIDE;
 	virtual bool OpenSourceFiles(const TArray<FString>& AbsoluteSourcePaths) OVERRIDE;
 	virtual bool SaveAllOpenDocuments() const OVERRIDE;
+	virtual void Tick(const float DeltaTime) OVERRIDE;
 };
\ No newline at end of file
diff --git a/Engine/Source/Developer/SourceCodeAccess/Private/SourceCodeAccessModule.cpp b/Engine/Source/Developer/SourceCodeAccess/Private/SourceCodeAccessModule.cpp
--- a/Engine/Source/Developer/SourceCodeAccess/Private/SourceCodeAccessModule.cpp
+++ b/Engine/Source/Developer/SourceCodeAccess/Private/SourceCodeAccessModule.cpp
@@ -97,11 +97,6 @@ FDoneLaunchingCodeAccessor& FSourceCodeAccessModule::OnDoneLaunchingCodeAccessor
 	return DoneLaunchingCodeAccessorDelegate;
 }
 
-FLaunchCodeAccessorDeferred& FSourceCodeAccessModule::OnLaunchCodeAccessorDeferred()
-{
-	return LaunchCodeAccessorDeferredDelegate;
-}
-
 FOpenFileFailed& FSourceCodeAccessModule::OnOpenFileFailed()
 {
 	return OpenFileFailedDelegate;
diff --git a/Engine/Source/Developer/SourceCodeAccess/Private/SourceCodeAccessModule.h b/Engine/Source/Developer/SourceCodeAccess/Private/SourceCodeAccessModule.h
--- a/Engine/Source/Developer/SourceCodeAccess/Private/SourceCodeAccessModule.h
+++ b/Engine/Source/Developer/SourceCodeAccess/Private/SourceCodeAccessModule.h
@@ -23,7 +23,6 @@ public:
 	virtual void SetAccessor(const FName& InName) OVERRIDE;
 	virtual FLaunchingCodeAccessor& OnLaunchingCodeAccessor() OVERRIDE;
 	virtual FDoneLaunchingCodeAccessor& OnDoneLaunchingCodeAccessor() OVERRIDE;
-	virtual FLaunchCodeAccessorDeferred& OnLaunchCodeAccessorDeferred() OVERRIDE;
 	virtual FOpenFileFailed& OnOpenFileFailed() OVERRIDE;
 
 private:
@@ -37,9 +36,6 @@ private:
 	/** Event delegate fired when done launching code accessor  */
 	FDoneLaunchingCodeAccessor DoneLaunchingCodeAccessorDelegate;
 
-	/** Event delegate fired when launching code accessor deferred */
-	FLaunchCodeAccessorDeferred LaunchCodeAccessorDeferredDelegate;
-
 	/** Event delegate fired when opening a file has failed */
 	FOpenFileFailed OpenFileFailedDelegate;
 
diff --git a/Engine/Source/Developer/SourceCodeAccess/Public/ISourceCodeAccessModule.h b/Engine/Source/Developer/SourceCodeAccess/Public/ISourceCodeAccessModule.h
--- a/Engine/Source/Developer/SourceCodeAccess/Public/ISourceCodeAccessModule.h
+++ b/Engine/Source/Developer/SourceCodeAccess/Public/ISourceCodeAccessModule.h
@@ -13,12 +13,6 @@ DECLARE_MULTICAST_DELEGATE( FLaunchingCodeAccessor );
  */
 DECLARE_MULTICAST_DELEGATE_OneParam( FDoneLaunchingCodeAccessor, const bool /*bSuccess*/);
 
-/**
- * Event fired when launching code accessor deferred
- * @param	InFilename	The command being deferred
- */
-DECLARE_MULTICAST_DELEGATE_OneParam( FLaunchCodeAccessorDeferred, const FString& /*InFilename*/);
-
 /**
  * Event fired when opening a file has failed
  * @param	InFilename	The filename that failed to open
@@ -55,9 +49,6 @@ public:
 	/** Gets the Event that is broadcasted when an attempted launch of this code accessor was successful or failed */
 	virtual FDoneLaunchingCodeAccessor& OnDoneLaunchingCodeAccessor() = 0;
 
-	/** Gets a delegate to be invoked when the the open command needs to be deferred */
-	virtual FLaunchCodeAccessorDeferred& OnLaunchCodeAccessorDeferred() = 0;
-
 	/** Gets the Event that is broadcast when an attempt to load a file through Visual Studio failed */
 	virtual FOpenFileFailed& OnOpenFileFailed() = 0;
 };
\ No newline at end of file
diff --git a/Engine/Source/Developer/SourceCodeAccess/Public/ISourceCodeAccessor.h b/Engine/Source/Developer/SourceCodeAccess/Public/ISourceCodeAccessor.h
--- a/Engine/Source/Developer/SourceCodeAccess/Public/ISourceCodeAccessor.h
+++ b/Engine/Source/Developer/SourceCodeAccess/Public/ISourceCodeAccessor.h
@@ -61,4 +61,10 @@ public:
 	 * @return true if successful
 	 */
 	virtual bool SaveAllOpenDocuments() const = 0;
+
+	/**
+	 * Tick this source code accessor
+	 * @param DeltaTime Delta time (in seconds) since the last call to Tick
+	 */
+	virtual void Tick(const float DeltaTime) = 0;
 };
\ No newline at end of file
diff --git a/Engine/Source/Editor/UnrealEd/Private/Editor.cpp b/Engine/Source/Editor/UnrealEd/Private/Editor.cpp
--- a/Engine/Source/Editor/UnrealEd/Private/Editor.cpp
+++ b/Engine/Source/Editor/UnrealEd/Private/Editor.cpp
@@ -30,6 +30,7 @@
 
 #include "AssetRegistryModule.h"
 #include "ContentBrowserModule.h"
+#include "ISourceCodeAccessModule.h"
 
 #include "Editor/MainFrame/Public/MainFrame.h"
 #include "AnimationUtils.h"
@@ -955,6 +956,8 @@ void UEditorEngine::Tick( float DeltaSeconds, bool bIdleMode )
 	// Tick the asset registry
 	FAssetRegistryModule::TickAssetRegistry(DeltaSeconds);
 
+	ISourceCodeAccessModule& SourceCodeAccessModule = FModuleManager::LoadModuleChecked<ISourceCodeAccessModule>("SourceCodeAccess");
+	SourceCodeAccessModule.GetAccessor().Tick(DeltaSeconds);
 
 	// tick the directory watcher
 	// @todo: Put me into an FTicker that is created when the DW module is loaded
diff --git a/Engine/Source/Editor/UnrealEd/Private/UnrealEdMisc.cpp b/Engine/Source/Editor/UnrealEd/Private/UnrealEdMisc.cpp
--- a/Engine/Source/Editor/UnrealEd/Private/UnrealEdMisc.cpp
+++ b/Engine/Source/Editor/UnrealEd/Private/UnrealEdMisc.cpp
@@ -11,7 +11,6 @@
 #include "MessageLog.h"
 #include "Developer/MessageLog/Public/MessageLogModule.h"
 #include "Kismet2/KismetDebugUtilities.h"
-#include "ISourceCodeAccessModule.h"
 #include "SourceControlWindows.h"
 #include "FbxLibs.h"
 #include "CompilerResultsLog.h"
@@ -93,10 +92,6 @@ void FUnrealEdMisc::OnInit()
 	FCoreDelegates::PreModal.AddRaw(this, &FUnrealEdMisc::OnEditorPreModal);
 	FCoreDelegates::PostModal.AddRaw(this, &FUnrealEdMisc::OnEditorPostModal);
 
-	/** Delegate that gets called by modules that can't directly access Engine */
-	ISourceCodeAccessModule& SourceCodeAccessModule = FModuleManager::LoadModuleChecked<ISourceCodeAccessModule>("SourceCodeAccess");
-	SourceCodeAccessModule.OnLaunchCodeAccessorDeferred().AddRaw(this, &FUnrealEdMisc::OnDeferCommand);
-
 	FEditorDelegates::OnAssetPostImport.AddStatic(&NormalMapIdentification::HandleAssetPostImport);
 
 	// Register the play world commands
@@ -539,9 +534,6 @@ void FUnrealEdMisc::OnExit()
 	FCoreDelegates::PreModal.RemoveAll(this);
 	FCoreDelegates::PostModal.RemoveAll(this);
 
-	ISourceCodeAccessModule& SourceCodeAccessModule = FModuleManager::LoadModuleChecked<ISourceCodeAccessModule>("SourceCodeAccess");
-	SourceCodeAccessModule.OnLaunchCodeAccessorDeferred().RemoveAll(this);
-
 	FComponentAssetBrokerage::PRIVATE_ShutdownBrokerage();
 
 	ISourceControlModule::Get().GetProvider().Close();
