Set the code accessor plugins to be Developer plugins rather than Editor only since Slate can be used within game now.
TTP #333675 - SLATE: WidgetReflector code reference hyperlinks don't work anymore
reviewer Steve.Robb

diff --git a/Engine/Plugins/Developer/VisualStudioSourceCodeAccess/VisualStudioSourceCodeAccess.uplugin b/Engine/Plugins/Developer/VisualStudioSourceCodeAccess/VisualStudioSourceCodeAccess.uplugin
--- a/Engine/Plugins/Developer/VisualStudioSourceCodeAccess/VisualStudioSourceCodeAccess.uplugin
+++ b/Engine/Plugins/Developer/VisualStudioSourceCodeAccess/VisualStudioSourceCodeAccess.uplugin
@@ -8,13 +8,13 @@
 	"CreatedByURL" : "http://epicgames.com",
 	"EngineVersion" : "4.2.0",
 	"Description" : "Allows access to source code in Visual Studio.",
-	"Category" : "Editor.Source Code Access",
+	"Category" : "Developer.Source Code Access",
 	
 	"Modules" :
 	[
 		{
 			"Name" : "VisualStudioSourceCodeAccess",
-			"Type" : "Editor",
+			"Type" : "Developer",
 			"WhitelistPlatforms" : [ "Win64", "Win32" ]
 		}
 	]
diff --git a/Engine/Plugins/Developer/XCodeSourceCodeAccess/XCodeSourceCodeAccess.uplugin b/Engine/Plugins/Developer/XCodeSourceCodeAccess/XCodeSourceCodeAccess.uplugin
--- a/Engine/Plugins/Developer/XCodeSourceCodeAccess/XCodeSourceCodeAccess.uplugin
+++ b/Engine/Plugins/Developer/XCodeSourceCodeAccess/XCodeSourceCodeAccess.uplugin
@@ -8,13 +8,13 @@
 	"CreatedByURL" : "http://epicgames.com",
 	"EngineVersion" : "4.2.0",
 	"Description" : "Allows access to source code in XCode.",
-	"Category" : "Editor.Source Code Access",
+	"Category" : "Developer.Source Code Access",
 	
 	"Modules" :
 	[
 		{
 			"Name" : "XCodeSourceCodeAccess",
-			"Type" : "Editor",
+			"Type" : "Developer",
 			"WhitelistPlatforms" : [ "Mac" ]
 		}
 	]
