Fix clicking on source links from blueprint projects in Rocket builds. Was not correctly handling no solution file being present. Also fix formatting from merges.

diff --git a/Engine/Plugins/Developer/VisualStudioSourceCodeAccess/Source/VisualStudioSourceCodeAccess/Private/VisualStudioSourceCodeAccessor.cpp b/Engine/Plugins/Developer/VisualStudioSourceCodeAccess/Source/VisualStudioSourceCodeAccess/Private/VisualStudioSourceCodeAccessor.cpp
--- a/Engine/Plugins/Developer/VisualStudioSourceCodeAccess/Source/VisualStudioSourceCodeAccess/Private/VisualStudioSourceCodeAccessor.cpp
+++ b/Engine/Plugins/Developer/VisualStudioSourceCodeAccess/Source/VisualStudioSourceCodeAccess/Private/VisualStudioSourceCodeAccessor.cpp
@@ -111,37 +111,37 @@ EAccessVisualStudioResult AccessVisualStudioViaDTE(CComPtr<EnvDTE::_DTE>& OutDTE
 		IEnumMoniker* MonikersTable;
 		if(SUCCEEDED(RunningObjectTable->EnumRunning(&MonikersTable)))
 		{
-		MonikersTable->Reset();
+			MonikersTable->Reset();
 
-		// Look for all visual studio instances in the ROT
-		IMoniker* CurrentMoniker;
+			// Look for all visual studio instances in the ROT
+			IMoniker* CurrentMoniker;
 			while(AccessResult != EAccessVisualStudioResult::VSInstanceIsOpen && MonikersTable->Next(1, &CurrentMoniker, NULL) == S_OK)
-		{
-			IBindCtx* BindContext;
-			LPOLESTR OutName;
+			{
+				IBindCtx* BindContext;
+				LPOLESTR OutName;
 				if(SUCCEEDED(CreateBindCtx(0, &BindContext)) && SUCCEEDED(CurrentMoniker->GetDisplayName(BindContext, NULL, &OutName)))
 				{
 					if(IsVisualStudioDTEMoniker(FString(OutName), InLocations))
 					{
-			CComPtr<IUnknown> ComObject;
+						CComPtr<IUnknown> ComObject;
 						if(SUCCEEDED(RunningObjectTable->GetObject(CurrentMoniker, &ComObject)))
-			{
-				CComPtr<EnvDTE::_DTE> TempDTE;
-				TempDTE = ComObject;
-
-				// Get the solution path for this instance
-				// If it equals the solution we would have opened above in RunVisualStudio(), we'll take that
-				CComPtr<EnvDTE::_Solution> Solution;
-				LPOLESTR OutPath;
-				if (SUCCEEDED(TempDTE->get_Solution(&Solution)) &&
-					SUCCEEDED(Solution->get_FullName(&OutPath)))
-				{
-					FString Filename(OutPath);
-					FPaths::NormalizeFilename(Filename);
+						{
+							CComPtr<EnvDTE::_DTE> TempDTE;
+							TempDTE = ComObject;
+
+							// Get the solution path for this instance
+							// If it equals the solution we would have opened above in RunVisualStudio(), we'll take that
+							CComPtr<EnvDTE::_Solution> Solution;
+							LPOLESTR OutPath;
+							if (SUCCEEDED(TempDTE->get_Solution(&Solution)) &&
+								SUCCEEDED(Solution->get_FullName(&OutPath)))
+							{
+								FString Filename(OutPath);
+								FPaths::NormalizeFilename(Filename);
 
-					if( Filename == InSolutionPath )
-					{
-						OutDTE = TempDTE;
+								if( Filename == InSolutionPath )
+								{
+									OutDTE = TempDTE;
 									AccessResult = EAccessVisualStudioResult::VSInstanceIsOpen;
 								}
 							}
@@ -157,11 +157,11 @@ EAccessVisualStudioResult AccessVisualStudioViaDTE(CComPtr<EnvDTE::_DTE>& OutDTE
 				{
 					UE_LOG(LogVSAccessor, Warning, TEXT("Couldn't get display name"));
 					AccessResult = EAccessVisualStudioResult::VSInstanceUnknown;
+				}
+				BindContext->Release();
+				CurrentMoniker->Release();
 			}
-			BindContext->Release();
-			CurrentMoniker->Release();
-		}
-		MonikersTable->Release();
+			MonikersTable->Release();
 		}
 		else
 		{
@@ -193,30 +193,30 @@ bool FVisualStudioSourceCodeAccessor::OpenVisualStudioSolutionViaDTE()
 	CComPtr<EnvDTE::_DTE> DTE;
 	switch (AccessVisualStudioViaDTE(DTE, GetSolutionPath(), Locations))
 	{
-		case EAccessVisualStudioResult::VSInstanceIsOpen:
-	{
-		// Set Focus on Visual Studio
-		CComPtr<EnvDTE::Window> MainWindow;
-		if (SUCCEEDED(DTE->get_MainWindow(&MainWindow)) &&
-			SUCCEEDED(MainWindow->Activate()))
-		{
-			bSuccess = true;
-		}
-		else
+	case EAccessVisualStudioResult::VSInstanceIsOpen:
 		{
-			UE_LOG(LogVSAccessor, Warning, TEXT("Couldn't set focus on Visual Studio."));
+			// Set Focus on Visual Studio
+			CComPtr<EnvDTE::Window> MainWindow;
+			if (SUCCEEDED(DTE->get_MainWindow(&MainWindow)) &&
+				SUCCEEDED(MainWindow->Activate()))
+			{
+				bSuccess = true;
+			}
+			else
+			{
+				UE_LOG(LogVSAccessor, Warning, TEXT("Couldn't set focus on Visual Studio."));
+			}
 		}
-	}
 		break;
 
 	case EAccessVisualStudioResult::VSInstanceIsNotOpen:
-	{
-		// Automatically fail if there's already an attempt in progress
-		if (!IsVSLaunchInProgress())
 		{
-			bSuccess = RunVisualStudioAndOpenSolution();
+			// Automatically fail if there's already an attempt in progress
+			if (!IsVSLaunchInProgress())
+			{
+				bSuccess = RunVisualStudioAndOpenSolution();
+			}
 		}
-	}
 		break;
 
 	default:
@@ -247,99 +247,99 @@ bool FVisualStudioSourceCodeAccessor::OpenVisualStudioFilesInternalViaDTE(const
 	switch (AccessVisualStudioViaDTE(DTE, GetSolutionPath(), Locations))
 	{
 	case EAccessVisualStudioResult::VSInstanceIsOpen:
-	{
-		// Set Focus on Visual Studio
-		CComPtr<EnvDTE::Window> MainWindow;
-		if (SUCCEEDED(DTE->get_MainWindow(&MainWindow)) &&
-			SUCCEEDED(MainWindow->Activate()))
 		{
-			// Get ItemOperations
-			CComPtr<EnvDTE::ItemOperations> ItemOperations;
-			if (SUCCEEDED(DTE->get_ItemOperations(&ItemOperations)))
+			// Set Focus on Visual Studio
+			CComPtr<EnvDTE::Window> MainWindow;
+			if (SUCCEEDED(DTE->get_MainWindow(&MainWindow)) &&
+				SUCCEEDED(MainWindow->Activate()))
 			{
-				for ( const FileOpenRequest& Request : Requests )
+				// Get ItemOperations
+				CComPtr<EnvDTE::ItemOperations> ItemOperations;
+				if (SUCCEEDED(DTE->get_ItemOperations(&ItemOperations)))
 				{
-					// Check that the file actually exists first
-					if ( !FPaths::FileExists(Request.FullPath) )
+					for ( const FileOpenRequest& Request : Requests )
 					{
-						SourceCodeAccessModule.OnOpenFileFailed().Broadcast(Request.FullPath);
-						bSuccess |= false;
-						continue;
-					}
+						// Check that the file actually exists first
+						if ( !FPaths::FileExists(Request.FullPath) )
+						{
+							SourceCodeAccessModule.OnOpenFileFailed().Broadcast(Request.FullPath);
+							bSuccess |= false;
+							continue;
+						}
 
-					// Open File
-					auto ANSIPath = StringCast<ANSICHAR>(*Request.FullPath);
-					CComBSTR COMStrFileName(ANSIPath.Get());
-					CComBSTR COMStrKind(EnvDTE::vsViewKindTextView);
-					CComPtr<EnvDTE::Window> Window;
-					if ( SUCCEEDED(ItemOperations->OpenFile(COMStrFileName, COMStrKind, &Window)) )
-					{
-						// If we've made it this far - we've opened the file.  it doesn't matter if
-						// we successfully get to the line number.  Everything else is gravy.
-						bSuccess |= true;
-
-						// Scroll to Line Number
-						CComPtr<EnvDTE::Document> Document;
-						CComPtr<IDispatch> SelectionDispatch;
-						CComPtr<EnvDTE::TextSelection> Selection;
-						if ( SUCCEEDED(DTE->get_ActiveDocument(&Document)) &&
-							 SUCCEEDED(Document->get_Selection(&SelectionDispatch)) &&
-							 SUCCEEDED(SelectionDispatch->QueryInterface(&Selection)) &&
-							 SUCCEEDED(Selection->GotoLine(Request.LineNumber, true)) )
+						// Open File
+						auto ANSIPath = StringCast<ANSICHAR>(*Request.FullPath);
+						CComBSTR COMStrFileName(ANSIPath.Get());
+						CComBSTR COMStrKind(EnvDTE::vsViewKindTextView);
+						CComPtr<EnvDTE::Window> Window;
+						if ( SUCCEEDED(ItemOperations->OpenFile(COMStrFileName, COMStrKind, &Window)) )
 						{
-							if ( !SUCCEEDED(Selection->MoveToLineAndOffset(Request.LineNumber, Request.ColumnNumber, false)) )
+							// If we've made it this far - we've opened the file.  it doesn't matter if
+							// we successfully get to the line number.  Everything else is gravy.
+							bSuccess |= true;
+
+							// Scroll to Line Number
+							CComPtr<EnvDTE::Document> Document;
+							CComPtr<IDispatch> SelectionDispatch;
+							CComPtr<EnvDTE::TextSelection> Selection;
+							if ( SUCCEEDED(DTE->get_ActiveDocument(&Document)) &&
+								 SUCCEEDED(Document->get_Selection(&SelectionDispatch)) &&
+								 SUCCEEDED(SelectionDispatch->QueryInterface(&Selection)) &&
+								 SUCCEEDED(Selection->GotoLine(Request.LineNumber, true)) )
+							{
+								if ( !SUCCEEDED(Selection->MoveToLineAndOffset(Request.LineNumber, Request.ColumnNumber, false)) )
+								{
+									UE_LOG(LogVSAccessor, Warning, TEXT("Couldn't goto column number '%i' of line '%i' in '%s'"), Request.ColumnNumber, Request.LineNumber, *Request.FullPath);
+								}
+							}
+							else
 							{
-								UE_LOG(LogVSAccessor, Warning, TEXT("Couldn't goto column number '%i' of line '%i' in '%s'"), Request.ColumnNumber, Request.LineNumber, *Request.FullPath);
+								UE_LOG(LogVSAccessor, Warning, TEXT("Couldn't goto line number '%i' in '%s'"), Request.LineNumber, *Request.FullPath);
 							}
 						}
 						else
 						{
-							UE_LOG(LogVSAccessor, Warning, TEXT("Couldn't goto line number '%i' in '%s'"), Request.LineNumber, *Request.FullPath);
+							UE_LOG(LogVSAccessor, Warning, TEXT("Couldn't open file '%s'."), *Request.FullPath);
 						}
 					}
-					else
-					{
-						UE_LOG(LogVSAccessor, Warning, TEXT("Couldn't open file '%s'."), *Request.FullPath);
-					}
-				}
 
-				VSLaunchFinished( true );
+					VSLaunchFinished( true );
+				}
+				else
+				{
+					UE_LOG(LogVSAccessor, Log, TEXT("Couldn't get item operations. Visual Studio may still be initializing."));
+					bDefer = true;
+				}
 			}
 			else
 			{
-				UE_LOG(LogVSAccessor, Log, TEXT("Couldn't get item operations. Visual Studio may still be initializing."));
-				bDefer = true;
+				UE_LOG(LogVSAccessor, Warning, TEXT("Couldn't set focus on Visual Studio."));
 			}
 		}
-		else
-		{
-			UE_LOG(LogVSAccessor, Warning, TEXT("Couldn't set focus on Visual Studio."));
-		}
-	}
-	break;
+		break;
 
 	case EAccessVisualStudioResult::VSInstanceIsNotOpen:
-	{
-		bDefer = true;
-
-		// We can't process until we're in the main thread, if we aren't initially defer until we are
-		if ( IsInGameThread() )
 		{
-			// If we haven't already attempted to launch VS do so now
-			if ( !IsVSLaunchInProgress() )
+			bDefer = true;
+
+			// We can't process until we're in the main thread, if we aren't initially defer until we are
+			if ( IsInGameThread() )
 			{
-				// If there's no valid instance of VS running, run one if we have it installed
-				if ( !RunVisualStudioAndOpenSolution() )
-				{
-					bDefer = false;
-				}
-				else
+				// If we haven't already attempted to launch VS do so now
+				if ( !IsVSLaunchInProgress() )
 				{
-					VSLaunchStarted();
+					// If there's no valid instance of VS running, run one if we have it installed
+					if ( !RunVisualStudioAndOpenSolution() )
+					{
+						bDefer = false;
+					}
+					else
+					{
+						VSLaunchStarted();
+					}
 				}
 			}
 		}
-	}
 		break;
 
 	default:
@@ -662,23 +662,23 @@ bool FVisualStudioSourceCodeAccessor::OpenVisualStudioSolutionViaProcess()
 	::DWORD ProcessID = 0;
 	FString Path;
 	switch (AccessVisualStudioViaProcess(ProcessID, Path, GetSolutionPath(), Locations))
-{
-	case EAccessVisualStudioResult::VSInstanceIsOpen:
 	{
-		// Try to bring Visual Studio to the foreground
-		::HWND VisualStudioHwnd = GetTopWindowForProcess(ProcessID);
-		if (VisualStudioHwnd)
+		case EAccessVisualStudioResult::VSInstanceIsOpen:
 		{
-			// SwitchToThisWindow isn't really intended for general use, however it can switch to 
-			// the VS window, where SetForegroundWindow will fail due to process permissions
-			::SwitchToThisWindow(VisualStudioHwnd, 0);
+			// Try to bring Visual Studio to the foreground
+			::HWND VisualStudioHwnd = GetTopWindowForProcess(ProcessID);
+			if (VisualStudioHwnd)
+			{
+				// SwitchToThisWindow isn't really intended for general use, however it can switch to 
+				// the VS window, where SetForegroundWindow will fail due to process permissions
+				::SwitchToThisWindow(VisualStudioHwnd, 0);
+			}
+			return true;
 		}
-		return true;
-	}
 		break;
 
 	case EAccessVisualStudioResult::VSInstanceIsNotOpen:
-	return RunVisualStudioAndOpenSolution();
+		return RunVisualStudioAndOpenSolution();
 
 	default:
 		// Do nothing if we failed the VS detection, otherwise we could get stuck in a loop of constantly 
@@ -874,18 +874,9 @@ bool FVisualStudioSourceCodeAccessor::RunVisualStudioAndOpenSolutionAndFiles(con
 		}
 	}
 
-	bool bSuccess = false;
-
-	if (!Params.IsEmpty())
-	{
-		FProcHandle WorkerHandle = FPlatformProcess::CreateProc(*ExecutablePath, *Params, true, false, false, nullptr, 0, nullptr, nullptr);
-		if (WorkerHandle.IsValid())
-		{
-			bSuccess = true;
-		}
-		WorkerHandle.Close();
-	}
-
+	FProcHandle WorkerHandle = FPlatformProcess::CreateProc(*ExecutablePath, *Params, true, false, false, nullptr, 0, nullptr, nullptr);
+	bool bSuccess = WorkerHandle.IsValid();
+	WorkerHandle.Close();
 	return bSuccess;
 }
 
@@ -915,7 +906,11 @@ FString FVisualStudioSourceCodeAccessor::GetSolutionPath() const
 	FScopeLock Lock(&CachedSolutionPathCriticalSection);
 	if(IsInGameThread())
 	{
-		CachedSolutionPath = FPaths::ConvertRelativePathToFull(FModuleManager::Get().GetSolutionFilepath());
+		FString SolutionPath = FModuleManager::Get().GetSolutionFilepath();
+		if(SolutionPath.Len() > 0)
+		{
+			CachedSolutionPath = FPaths::ConvertRelativePathToFull(SolutionPath);
+		}
 	}
 	return CachedSolutionPath;
 }
