Fix crash on shutdown if hot reload module hasn't been loaded yet

diff --git a/Engine/Plugins/Developer/VisualStudioSourceCodeAccess/Source/VisualStudioSourceCodeAccess/Private/VisualStudioSourceCodeAccessor.cpp b/Engine/Plugins/Developer/VisualStudioSourceCodeAccess/Source/VisualStudioSourceCodeAccess/Private/VisualStudioSourceCodeAccessor.cpp
--- a/Engine/Plugins/Developer/VisualStudioSourceCodeAccess/Source/VisualStudioSourceCodeAccess/Private/VisualStudioSourceCodeAccessor.cpp
+++ b/Engine/Plugins/Developer/VisualStudioSourceCodeAccess/Source/VisualStudioSourceCodeAccess/Private/VisualStudioSourceCodeAccessor.cpp
@@ -84,7 +84,10 @@ void FVisualStudioSourceCodeAccessor::Startup()
 void FVisualStudioSourceCodeAccessor::Shutdown()
 {
 	// Unregister the hot-reload callback
-	IHotReloadModule::Get().OnModuleCompilerStarted().RemoveStatic( &SaveVisualStudioDocuments );
+	if(IHotReloadModule::IsAvailable())
+	{
+		IHotReloadModule::Get().OnModuleCompilerStarted().RemoveStatic( &SaveVisualStudioDocuments );
+	}
 }
 
 #if VSACCESSOR_HAS_DTE
