[INTEGRATE] Change 2425885 by [REDACTED_PERSON][REDACTED_HANDLE] on 2015/01/30 11:13:42

	Allowed the DTE VS integration to retry a file open operation if it detected a blocked instance of VS

	UE-8623 - Adding new code to project opens new Visual Studio rather than using existing one

	We'd previously disabled the defered opening (essentially a retry) of files if VS detection failed as it could get stuck in a loop when using a non-DTE version of VS (VS Express), however there are timing issues when we add new classes in the editor where an existing VS instance can be blocked waiting on a modal dialog.

	I've changed our VS integration to report these DTE based VS instances as "blocked" rather than "unknown", and we now allow a defer on these until they become unblocked.

diff --git a/Engine/Plugins/Developer/VisualStudioSourceCodeAccess/Source/VisualStudioSourceCodeAccess/Private/VisualStudioSourceCodeAccessor.cpp b/Engine/Plugins/Developer/VisualStudioSourceCodeAccess/Source/VisualStudioSourceCodeAccess/Private/VisualStudioSourceCodeAccessor.cpp
--- a/Engine/Plugins/Developer/VisualStudioSourceCodeAccess/Source/VisualStudioSourceCodeAccess/Private/VisualStudioSourceCodeAccessor.cpp
+++ b/Engine/Plugins/Developer/VisualStudioSourceCodeAccess/Source/VisualStudioSourceCodeAccess/Private/VisualStudioSourceCodeAccessor.cpp
@@ -54,6 +54,8 @@ enum class EAccessVisualStudioResult : uint8
 	VSInstanceIsOpen,
 	/** An instance of Visual Studio is not available */
 	VSInstanceIsNotOpen,
+	/** An instance of Visual Studio is open, but could not be fully queried because it is blocked by a modal operation - this may succeed later */
+	VSInstanceIsBlocked,
 	/** It is unknown whether an instance of Visual Studio is available, as an error occurred when performing the check */
 	VSInstanceUnknown,
 };
@@ -153,6 +155,11 @@ EAccessVisualStudioResult AccessVisualStudioViaDTE(CComPtr<EnvDTE::_DTE>& OutDTE
 									AccessResult = EAccessVisualStudioResult::VSInstanceIsOpen;
 								}
 							}
+							else
+							{
+								UE_LOG(LogVSAccessor, Warning, TEXT("Visual Studio is open but could not be queried - it may be blocked by a modal operation"));
+								AccessResult = EAccessVisualStudioResult::VSInstanceIsBlocked;
+							}
 						}
 						else
 						{
@@ -350,6 +357,14 @@ bool FVisualStudioSourceCodeAccessor::OpenVisualStudioFilesInternalViaDTE(const
 		}
 		break;
 
+	case EAccessVisualStudioResult::VSInstanceIsBlocked:
+		{
+			// VS may be open for the solution we want, but we can't query it right now as it's blocked for some reason
+			// Defer this operation so we can try it again later should VS become unblocked
+			bDefer = true;
+		}
+		break;
+
 	default:
 		// Do nothing if we failed the VS detection, otherwise we could get stuck in a loop of constantly 
 		// trying to open a VS instance since we can't detect that one is already running
@@ -925,17 +940,17 @@ FString FVisualStudioSourceCodeAccessor::GetSolutionPath() const
 
 void FVisualStudioSourceCodeAccessor::Tick(const float DeltaTime)
 {
-	if (IsVSLaunchInProgress())
+	TArray<FileOpenRequest> TmpDeferredRequests;
 	{
-		TArray<FileOpenRequest> TmpDeferredRequests;
-		{
-			FScopeLock Lock(&DeferredRequestsCriticalSection);
+		FScopeLock Lock(&DeferredRequestsCriticalSection);
 
-			// Copy the DeferredRequests array, as OpenVisualStudioFilesInternal may update it
-			TmpDeferredRequests = DeferredRequests;
-			DeferredRequests.Empty();
-		}
+		// Copy the DeferredRequests array, as OpenVisualStudioFilesInternal may update it
+		TmpDeferredRequests = DeferredRequests;
+		DeferredRequests.Empty();
+	}
 
+	if(TmpDeferredRequests.Num() > 0)
+	{
 		// Try and open any pending files in VS first (this will update the VS launch state appropriately)
 		OpenVisualStudioFilesInternal(TmpDeferredRequests);
 	}
