[INTEGRATE] Change 2428300 by [REDACTED_PERSON][REDACTED_HANDLE] on 2015/02/02 12:25:57

	Added a way to use DTE to add new files directly into VS without having to generate projects

	This avoids the annoying "Do you want to reload?" dialog that appears in VS when projects are externally changed, and also saves us the cost/time of having to run UBT to generate projects.

diff --git a/Engine/Plugins/Developer/KDevelopSourceCodeAccess/Source/KDevelopSourceCodeAccess/Private/KDevelopSourceCodeAccessor.cpp b/Engine/Plugins/Developer/KDevelopSourceCodeAccess/Source/KDevelopSourceCodeAccess/Private/KDevelopSourceCodeAccessor.cpp
--- a/Engine/Plugins/Developer/KDevelopSourceCodeAccess/Source/KDevelopSourceCodeAccess/Private/KDevelopSourceCodeAccessor.cpp
+++ b/Engine/Plugins/Developer/KDevelopSourceCodeAccess/Source/KDevelopSourceCodeAccess/Private/KDevelopSourceCodeAccessor.cpp
@@ -70,6 +70,11 @@ bool FKDevelopSourceCodeAccessor::OpenSourceFiles(const TArray<FString>& Absolut
 	return false;
 }
 
+bool FKDevelopSourceCodeAccessor::AddSourceFiles(const TArray<FString>& AbsoluteSourcePaths, const TArray<FString>& AvailableModules)
+{
+	return false;
+}
+
 bool FKDevelopSourceCodeAccessor::OpenFileAtLine(const FString& FullPath, int32 LineNumber, int32 ColumnNumber)
 {
 	// column & line numbers are 1-based, so dont allow zero
diff --git a/Engine/Plugins/Developer/KDevelopSourceCodeAccess/Source/KDevelopSourceCodeAccess/Private/KDevelopSourceCodeAccessor.h b/Engine/Plugins/Developer/KDevelopSourceCodeAccess/Source/KDevelopSourceCodeAccess/Private/KDevelopSourceCodeAccessor.h
--- a/Engine/Plugins/Developer/KDevelopSourceCodeAccess/Source/KDevelopSourceCodeAccess/Private/KDevelopSourceCodeAccessor.h
+++ b/Engine/Plugins/Developer/KDevelopSourceCodeAccess/Source/KDevelopSourceCodeAccess/Private/KDevelopSourceCodeAccessor.h
@@ -20,6 +20,7 @@ public:
 	virtual bool OpenSolution() override;
 	virtual bool OpenFileAtLine(const FString& FullPath, int32 LineNumber, int32 ColumnNumber = 0) override;
 	virtual bool OpenSourceFiles(const TArray<FString>& AbsoluteSourcePaths) override;
+	virtual bool AddSourceFiles(const TArray<FString>& AbsoluteSourcePaths, const TArray<FString>& AvailableModules) override;
 	virtual bool SaveAllOpenDocuments() const override;
 	virtual void Tick(const float DeltaTime) override;
 
diff --git a/Engine/Plugins/Developer/VisualStudioSourceCodeAccess/Source/VisualStudioSourceCodeAccess/Private/VisualStudioSourceCodeAccessor.cpp b/Engine/Plugins/Developer/VisualStudioSourceCodeAccess/Source/VisualStudioSourceCodeAccess/Private/VisualStudioSourceCodeAccessor.cpp
--- a/Engine/Plugins/Developer/VisualStudioSourceCodeAccess/Source/VisualStudioSourceCodeAccess/Private/VisualStudioSourceCodeAccessor.cpp
+++ b/Engine/Plugins/Developer/VisualStudioSourceCodeAccess/Source/VisualStudioSourceCodeAccess/Private/VisualStudioSourceCodeAccessor.cpp
@@ -778,6 +778,191 @@ bool FVisualStudioSourceCodeAccessor::OpenSourceFiles(const TArray<FString>& Abs
 	return false;
 }
 
+bool FVisualStudioSourceCodeAccessor::AddSourceFiles(const TArray<FString>& AbsoluteSourcePaths, const TArray<FString>& AvailableModules)
+{
+	// This requires DTE - there is no fallback for this operation when DTE is not available
+#if VSACCESSOR_HAS_DTE
+	bool bSuccess = true;
+
+	struct FModuleNameAndPath
+	{
+		FString ModuleBuildFilePath;
+		FString ModulePath;
+		FName ModuleName;
+	};
+
+	TArray<FModuleNameAndPath> ModuleNamesAndPaths;
+	ModuleNamesAndPaths.Reserve(AvailableModules.Num());
+	for (const FString& AvailableModule : AvailableModules)
+	{
+		static const int32 BuildFileExtensionLen = FString(TEXT(".Build.cs")).Len();
+
+		// AvailableModule is the relative path to the .Build.cs file
+		FModuleNameAndPath ModuleNameAndPath;
+		ModuleNameAndPath.ModuleBuildFilePath = FPaths::ConvertRelativePathToFull(AvailableModule);
+		ModuleNameAndPath.ModulePath = FPaths::GetPath(ModuleNameAndPath.ModuleBuildFilePath);
+		ModuleNameAndPath.ModuleName = *FPaths::GetCleanFilename(ModuleNameAndPath.ModuleBuildFilePath).LeftChop(BuildFileExtensionLen);
+		ModuleNamesAndPaths.Add(ModuleNameAndPath);
+	}
+
+	struct FModuleNewSourceFiles
+	{
+		FModuleNameAndPath ModuleNameAndPath;
+		TArray<FString> NewSourceFiles;
+	};
+
+	// Work out which module each source file will be in
+	TMap<FName, FModuleNewSourceFiles> ModuleToNewSourceFiles;
+	{
+		const FModuleNameAndPath* LastSourceFilesModule = nullptr;
+		for (const FString& SourceFile : AbsoluteSourcePaths)
+		{
+			// First check to see if this source file is in the same module as the last source file - this is usually the case, and saves us a lot of string compares
+			if (LastSourceFilesModule && SourceFile.StartsWith(LastSourceFilesModule->ModulePath))
+			{
+				FModuleNewSourceFiles& ModuleNewSourceFiles = ModuleToNewSourceFiles.FindChecked(LastSourceFilesModule->ModuleName);
+				ModuleNewSourceFiles.NewSourceFiles.Add(SourceFile);
+				continue;
+			}
+
+			// Look for the module which will contain this file
+			LastSourceFilesModule = nullptr;
+			for (const FModuleNameAndPath& ModuleNameAndPath : ModuleNamesAndPaths)
+			{
+				if (SourceFile.StartsWith(ModuleNameAndPath.ModulePath))
+				{
+					LastSourceFilesModule = &ModuleNameAndPath;
+
+					FModuleNewSourceFiles& ModuleNewSourceFiles = ModuleToNewSourceFiles.FindOrAdd(ModuleNameAndPath.ModuleName);
+					ModuleNewSourceFiles.ModuleNameAndPath = ModuleNameAndPath;
+					ModuleNewSourceFiles.NewSourceFiles.Add(SourceFile);
+					break;
+				}
+			}
+
+			// Failed to find the module for this source file?
+			if (!LastSourceFilesModule)
+			{
+				UE_LOG(LogVSAccessor, Warning, TEXT("Cannot add source file '%s' as it doesn't belong to a known module"), *SourceFile);
+				bSuccess = false;
+			}
+		}
+	}
+
+	CComPtr<EnvDTE::_DTE> DTE;
+	if (AccessVisualStudioViaDTE(DTE, GetSolutionPath(), Locations) == EAccessVisualStudioResult::VSInstanceIsOpen)
+	{
+		CComPtr<EnvDTE::_Solution> Solution;
+		if (SUCCEEDED(DTE->get_Solution(&Solution)) && Solution)
+		{
+			// Process each module
+			for (const auto& ModuleNewSourceFilesKeyValue : ModuleToNewSourceFiles)
+			{
+				const FModuleNewSourceFiles& ModuleNewSourceFiles = ModuleNewSourceFilesKeyValue.Value;
+
+				const FString& ModuleBuildFilePath = ModuleNewSourceFiles.ModuleNameAndPath.ModuleBuildFilePath;
+				auto ANSIModuleBuildFilePath = StringCast<ANSICHAR>(*ModuleBuildFilePath);
+				CComBSTR COMStrModuleBuildFilePath(ANSIModuleBuildFilePath.Get());
+
+				CComPtr<EnvDTE::ProjectItem> BuildFileProjectItem;
+				if (SUCCEEDED(Solution->FindProjectItem(COMStrModuleBuildFilePath, &BuildFileProjectItem)) && BuildFileProjectItem)
+				{
+					// We found the .Build.cs file in the existing solution - now we need its parent ProjectItems as that's what we'll be adding new content to
+					CComPtr<EnvDTE::ProjectItems> ModuleProjectFolder;
+					if (SUCCEEDED(BuildFileProjectItem->get_Collection(&ModuleProjectFolder)) && ModuleProjectFolder)
+					{
+						for (const FString& SourceFile : AbsoluteSourcePaths)
+						{
+							const FString ProjectRelativeSourceFilePath = SourceFile.Mid(ModuleNewSourceFiles.ModuleNameAndPath.ModulePath.Len());
+							TArray<FString> SourceFileParts;
+							ProjectRelativeSourceFilePath.ParseIntoArray(&SourceFileParts, TEXT("/"), true);
+					
+							if (SourceFileParts.Num() == 0)
+							{
+								// This should never happen as it means we somehow have no filename within the project directory
+								bSuccess = false;
+								continue;
+							}
+
+							CComPtr<EnvDTE::ProjectItems> CurProjectItems = ModuleProjectFolder;
+
+							// Firstly we need to make sure that all the folders we need exist - this also walks us down to the correct place to add the file
+							for (int32 FilePartIndex = 0; FilePartIndex < SourceFileParts.Num() - 1 && CurProjectItems; ++FilePartIndex)
+							{
+								const FString& SourceFilePart = SourceFileParts[FilePartIndex];
+
+								auto ANSIPart = StringCast<ANSICHAR>(*SourceFilePart);
+								CComBSTR COMStrFilePart(ANSIPart.Get());
+
+								::VARIANT vProjectItemName;
+								vProjectItemName.vt = VT_BSTR;
+								vProjectItemName.bstrVal = COMStrFilePart;
+
+								CComPtr<EnvDTE::ProjectItem> ProjectItem;
+								if (SUCCEEDED(CurProjectItems->Item(vProjectItemName, &ProjectItem)) && !ProjectItem)
+								{
+									// Add this part
+									CurProjectItems->AddFolder(COMStrFilePart, nullptr, &ProjectItem);
+								}
+
+								if (ProjectItem)
+								{
+									ProjectItem->get_ProjectItems(&CurProjectItems);
+								}
+								else
+								{
+									CurProjectItems = nullptr;
+								}
+							}
+
+							if (!CurProjectItems)
+							{
+								// Failed to find or add all the path parts
+								bSuccess = false;
+								continue;
+							}
+
+							// Now we add the file to the project under the last folder we found along its path
+							auto ANSIPath = StringCast<ANSICHAR>(*SourceFile);
+							CComBSTR COMStrFileName(ANSIPath.Get());
+							CComPtr<EnvDTE::ProjectItem> FileProjectItem;
+							if (SUCCEEDED(CurProjectItems->AddFromFile(COMStrFileName, &FileProjectItem)))
+							{
+								bSuccess &= true;
+							}
+						}
+					}
+					else
+					{
+						UE_LOG(LogVSAccessor, Warning, TEXT("Cannot add source files as we failed to get the parent items container for the '%s' item"), *ModuleBuildFilePath);
+						bSuccess = false;
+					}
+				}
+				else
+				{
+					UE_LOG(LogVSAccessor, Warning, TEXT("Cannot add source files as we failed to find '%s' in the solution"), *ModuleBuildFilePath);
+					bSuccess = false;
+				}
+			}
+		}
+		else
+		{
+			UE_LOG(LogVSAccessor, Warning, TEXT("Cannot add source files as Visual Studio failed to return a solution when queried"));
+			bSuccess = false;
+		}
+	}
+	else
+	{
+		UE_LOG(LogVSAccessor, Verbose, TEXT("Cannot add source files as Visual Studio is either not open or not responding"));
+		bSuccess = false;
+	}
+
+	return bSuccess;
+#endif
+
+	return false;
+}
+
 bool FVisualStudioSourceCodeAccessor::OpenFileAtLine(const FString& FullPath, int32 LineNumber, int32 ColumnNumber)
 {
 	// column & line numbers are 1-based, so dont allow zero
diff --git a/Engine/Plugins/Developer/VisualStudioSourceCodeAccess/Source/VisualStudioSourceCodeAccess/Private/VisualStudioSourceCodeAccessor.h b/Engine/Plugins/Developer/VisualStudioSourceCodeAccess/Source/VisualStudioSourceCodeAccess/Private/VisualStudioSourceCodeAccessor.h
--- a/Engine/Plugins/Developer/VisualStudioSourceCodeAccess/Source/VisualStudioSourceCodeAccess/Private/VisualStudioSourceCodeAccessor.h
+++ b/Engine/Plugins/Developer/VisualStudioSourceCodeAccess/Source/VisualStudioSourceCodeAccess/Private/VisualStudioSourceCodeAccessor.h
@@ -33,6 +33,7 @@ public:
 	virtual bool OpenSolution() override;
 	virtual bool OpenFileAtLine(const FString& FullPath, int32 LineNumber, int32 ColumnNumber = 0) override;
 	virtual bool OpenSourceFiles(const TArray<FString>& AbsoluteSourcePaths) override;
+	virtual bool AddSourceFiles(const TArray<FString>& AbsoluteSourcePaths, const TArray<FString>& AvailableModules) override;
 	virtual bool SaveAllOpenDocuments() const override;
 	virtual void Tick(const float DeltaTime) override;
 
diff --git a/Engine/Plugins/Developer/XCodeSourceCodeAccess/Source/XCodeSourceCodeAccess/Private/XCodeSourceCodeAccessor.cpp b/Engine/Plugins/Developer/XCodeSourceCodeAccess/Source/XCodeSourceCodeAccess/Private/XCodeSourceCodeAccessor.cpp
--- a/Engine/Plugins/Developer/XCodeSourceCodeAccess/Source/XCodeSourceCodeAccess/Private/XCodeSourceCodeAccessor.cpp
+++ b/Engine/Plugins/Developer/XCodeSourceCodeAccess/Source/XCodeSourceCodeAccess/Private/XCodeSourceCodeAccessor.cpp
@@ -184,6 +184,11 @@ bool FXCodeSourceCodeAccessor::OpenSourceFiles(const TArray<FString>& AbsoluteSo
 	return true;
 }
 
+bool FXCodeSourceCodeAccessor::AddSourceFiles(const TArray<FString>& AbsoluteSourcePaths, const TArray<FString>& AvailableModules)
+{
+	return false;
+}
+
 bool FXCodeSourceCodeAccessor::SaveAllOpenDocuments() const
 {
 	bool ExecutionSucceeded = false;
diff --git a/Engine/Plugins/Developer/XCodeSourceCodeAccess/Source/XCodeSourceCodeAccess/Private/XCodeSourceCodeAccessor.h b/Engine/Plugins/Developer/XCodeSourceCodeAccess/Source/XCodeSourceCodeAccess/Private/XCodeSourceCodeAccessor.h
--- a/Engine/Plugins/Developer/XCodeSourceCodeAccess/Source/XCodeSourceCodeAccess/Private/XCodeSourceCodeAccessor.h
+++ b/Engine/Plugins/Developer/XCodeSourceCodeAccess/Source/XCodeSourceCodeAccess/Private/XCodeSourceCodeAccessor.h
@@ -15,6 +15,7 @@ public:
 	virtual bool OpenSolution() override;
 	virtual bool OpenFileAtLine(const FString& FullPath, int32 LineNumber, int32 ColumnNumber = 0) override;
 	virtual bool OpenSourceFiles(const TArray<FString>& AbsoluteSourcePaths) override;
+	virtual bool AddSourceFiles(const TArray<FString>& AbsoluteSourcePaths, const TArray<FString>& AvailableModules) override;
 	virtual bool SaveAllOpenDocuments() const override;
 	virtual void Tick(const float DeltaTime) override;
 };
\ No newline at end of file
diff --git a/Engine/Source/Developer/SourceCodeAccess/Private/DefaultSourceCodeAccessor.cpp b/Engine/Source/Developer/SourceCodeAccess/Private/DefaultSourceCodeAccessor.cpp
--- a/Engine/Source/Developer/SourceCodeAccess/Private/DefaultSourceCodeAccessor.cpp
+++ b/Engine/Source/Developer/SourceCodeAccess/Private/DefaultSourceCodeAccessor.cpp
@@ -40,6 +40,11 @@ bool FDefaultSourceCodeAccessor::OpenSourceFiles(const TArray<FString>& Absolute
 	return false;
 }
 
+bool FDefaultSourceCodeAccessor::AddSourceFiles(const TArray<FString>& AbsoluteSourcePaths, const TArray<FString>& AvailableModules)
+{
+	return false;
+}
+
 bool FDefaultSourceCodeAccessor::SaveAllOpenDocuments() const
 {
 	return false;
diff --git a/Engine/Source/Developer/SourceCodeAccess/Private/DefaultSourceCodeAccessor.h b/Engine/Source/Developer/SourceCodeAccess/Private/DefaultSourceCodeAccessor.h
--- a/Engine/Source/Developer/SourceCodeAccess/Private/DefaultSourceCodeAccessor.h
+++ b/Engine/Source/Developer/SourceCodeAccess/Private/DefaultSourceCodeAccessor.h
@@ -15,6 +15,7 @@ public:
 	virtual bool OpenSolution() override;
 	virtual bool OpenFileAtLine(const FString& FullPath, int32 LineNumber, int32 ColumnNumber = 0) override;
 	virtual bool OpenSourceFiles(const TArray<FString>& AbsoluteSourcePaths) override;
+	virtual bool AddSourceFiles(const TArray<FString>& AbsoluteSourcePaths, const TArray<FString>& AvailableModules) override;
 	virtual bool SaveAllOpenDocuments() const override;
 	virtual void Tick(const float DeltaTime) override;
 };
\ No newline at end of file
diff --git a/Engine/Source/Developer/SourceCodeAccess/Public/ISourceCodeAccessor.h b/Engine/Source/Developer/SourceCodeAccess/Public/ISourceCodeAccessor.h
--- a/Engine/Source/Developer/SourceCodeAccess/Public/ISourceCodeAccessor.h
+++ b/Engine/Source/Developer/SourceCodeAccess/Public/ISourceCodeAccessor.h
@@ -55,6 +55,14 @@ public:
 	 */
 	virtual bool OpenSourceFiles(const TArray<FString>& AbsoluteSourcePaths) = 0;
 
+	/**
+	 * Add a group of source files to the current solution/project/workspace
+	 * @param	AbsoluteSourcePaths		Array of paths to files to open
+	 * @param	AvailableModules		Array of known module locations (.Build.cs files) - you can get this from calling FSourceCodeNavigation::GetSourceFileDatabase().GetModuleNames() if in the editor
+	 * @return true if the files could be added, false otherwise
+	 */
+	virtual bool AddSourceFiles(const TArray<FString>& AbsoluteSourcePaths, const TArray<FString>& AvailableModules) = 0;
+
 	/**
 	 * Saves all open code documents if they need to be saved.
 	 * Will block if there is any read-only files open that need to be saved.
diff --git a/Engine/Source/Editor/GameProjectGeneration/Private/GameProjectUtils.cpp b/Engine/Source/Editor/GameProjectGeneration/Private/GameProjectUtils.cpp
--- a/Engine/Source/Editor/GameProjectGeneration/Private/GameProjectUtils.cpp
+++ b/Engine/Source/Editor/GameProjectGeneration/Private/GameProjectUtils.cpp
@@ -2986,8 +2986,8 @@ bool GameProjectUtils::AddCodeToProject_Internal(const FString& NewClassName, co
 	
 	// If the project does not already contain code, add the primary game module
 	TArray<FString> CreatedFiles;
-	bool bDidNotHaveAnyCodeFiles = !ProjectHasCodeFiles();
-	if (bDidNotHaveAnyCodeFiles)
+	const bool bProjectHadCodeFiles = ProjectHasCodeFiles();
+	if (!bProjectHadCodeFiles)
 	{
 		// We always add the basic source code to the root directory, not the potential sub-directory provided by NewClassPath
 		const FString SourceDir = FPaths::GameSourceDir().LeftChop(1); // Trim the trailing /
@@ -3043,11 +3043,26 @@ bool GameProjectUtils::AddCodeToProject_Internal(const FString& NewClassName, co
 
 	SlowTask.EnterProgressFrame();
 
-	// Generate project files if we happen to be using a project file.
-	if ( !FDesktopPlatformModule::Get()->GenerateProjectFiles(FPaths::RootDir(), FPaths::GetProjectFilePath(), GWarn) )
+	TArray<FString> CreatedFilesForExternalAppRead;
+	CreatedFilesForExternalAppRead.Reserve(CreatedFiles.Num());
+	for (const FString& CreatedFile : CreatedFiles)
 	{
-		OutFailReason = LOCTEXT("FailedToGenerateProjectFiles", "Failed to generate project files.");
-		return false;
+		CreatedFilesForExternalAppRead.Add( IFileManager::Get().ConvertToAbsolutePathForExternalAppForRead(*CreatedFile) );
+	}
+
+	// First see if we can avoid a full generation by adding the new files to an already open project
+	if ( bProjectHadCodeFiles && FSourceCodeNavigation::AddSourceFiles(CreatedFilesForExternalAppRead) )
+	{
+		// todo: jdale - even if this succeeds, we still need to run UBT with "-gather" to update the makefiles
+	}
+	else
+	{
+		// Generate project files if we happen to be using a project file.
+		if ( !FDesktopPlatformModule::Get()->GenerateProjectFiles(FPaths::RootDir(), FPaths::GetProjectFilePath(), GWarn) )
+		{
+			OutFailReason = LOCTEXT("FailedToGenerateProjectFiles", "Failed to generate project files.");
+			return false;
+		}
 	}
 
 	SlowTask.EnterProgressFrame();
@@ -3056,13 +3071,7 @@ bool GameProjectUtils::AddCodeToProject_Internal(const FString& NewClassName, co
 	ISourceControlProvider& SourceControlProvider = ISourceControlModule::Get().GetProvider();
 	if ( ISourceControlModule::Get().IsEnabled() && SourceControlProvider.IsAvailable() )
 	{
-		TArray<FString> FilesToCheckOut;
-		for ( auto FileIt = CreatedFiles.CreateConstIterator(); FileIt; ++FileIt )
-		{
-			FilesToCheckOut.Add( IFileManager::Get().ConvertToAbsolutePathForExternalAppForRead(**FileIt) );
-		}
-
-		SourceControlProvider.Execute(ISourceControlOperation::Create<FMarkForAdd>(), FilesToCheckOut);
+		SourceControlProvider.Execute(ISourceControlOperation::Create<FMarkForAdd>(), CreatedFilesForExternalAppRead);
 	}
 
 	SlowTask.EnterProgressFrame( 1.0f, LOCTEXT("CompilingCPlusPlusCode", "Compiling new C++ code.  Please wait..."));
@@ -3070,7 +3079,7 @@ bool GameProjectUtils::AddCodeToProject_Internal(const FString& NewClassName, co
 	OutHeaderFilePath = NewHeaderFilename;
 	OutCppFilePath = NewCppFilename;
 
-	if (bDidNotHaveAnyCodeFiles)
+	if (!bProjectHadCodeFiles)
 	{
 		// This is the first time we add code to this project so compile its game DLL
 		const FString GameModuleName = FApp::GetGameName();
diff --git a/Engine/Source/Editor/UnrealEd/Private/SourceCodeNavigation.cpp b/Engine/Source/Editor/UnrealEd/Private/SourceCodeNavigation.cpp
--- a/Engine/Source/Editor/UnrealEd/Private/SourceCodeNavigation.cpp
+++ b/Engine/Source/Editor/UnrealEd/Private/SourceCodeNavigation.cpp
@@ -1548,9 +1548,7 @@ bool FSourceCodeNavigation::OpenSourceFiles(const TArray<FString>& AbsoluteSourc
 	if ( IsCompilerAvailable() )
 	{
 		ISourceCodeAccessModule& SourceCodeAccessModule = FModuleManager::LoadModuleChecked<ISourceCodeAccessModule>("SourceCodeAccess");
-		SourceCodeAccessModule.GetAccessor().OpenSourceFiles(AbsoluteSourcePaths);
-
-		return true;
+		return SourceCodeAccessModule.GetAccessor().OpenSourceFiles(AbsoluteSourcePaths);
 	}
 
 	// Let others know that we've failed to open some source files.
@@ -1559,6 +1557,17 @@ bool FSourceCodeNavigation::OpenSourceFiles(const TArray<FString>& AbsoluteSourc
 	return false;
 }
 
+bool FSourceCodeNavigation::AddSourceFiles(const TArray<FString>& AbsoluteSourcePaths)
+{
+	if ( IsCompilerAvailable() )
+	{
+		ISourceCodeAccessModule& SourceCodeAccessModule = FModuleManager::LoadModuleChecked<ISourceCodeAccessModule>("SourceCodeAccess");
+		return SourceCodeAccessModule.GetAccessor().AddSourceFiles(AbsoluteSourcePaths, GetSourceFileDatabase().GetModuleNames());
+	}
+
+	return false;
+}
+
 bool FSourceCodeNavigation::OpenModuleSolution()
 {
 	ISourceCodeAccessModule& SourceCodeAccessModule = FModuleManager::LoadModuleChecked<ISourceCodeAccessModule>("SourceCodeAccess");
diff --git a/Engine/Source/Editor/UnrealEd/Public/SourceCodeNavigation.h b/Engine/Source/Editor/UnrealEd/Public/SourceCodeNavigation.h
--- a/Engine/Source/Editor/UnrealEd/Public/SourceCodeNavigation.h
+++ b/Engine/Source/Editor/UnrealEd/Public/SourceCodeNavigation.h
@@ -178,6 +178,9 @@ public:
 	/** Opens a multiple source files */
 	UNREALED_API static bool OpenSourceFiles(const TArray<FString>& AbsoluteSourcePaths);
 
+	/** Add multiple source files to the current solution/project/workspace */
+	UNREALED_API static bool AddSourceFiles(const TArray<FString>& AbsoluteSourcePaths);
+
 	UNREALED_API static bool OpenModuleSolution();
 
 	/** Attempt to locate fully qualified class module name */
