Stopped VS showing a "Save" message when closing after adding new classes via the editor

We now save the project when we add new files to via the editor. This avoids VS showing a prompt about the .vcxproj being dirty when closed.

diff --git a/Engine/Plugins/Developer/VisualStudioSourceCodeAccess/Source/VisualStudioSourceCodeAccess/Private/VisualStudioSourceCodeAccessor.cpp b/Engine/Plugins/Developer/VisualStudioSourceCodeAccess/Source/VisualStudioSourceCodeAccess/Private/VisualStudioSourceCodeAccessor.cpp
--- a/Engine/Plugins/Developer/VisualStudioSourceCodeAccess/Source/VisualStudioSourceCodeAccess/Private/VisualStudioSourceCodeAccessor.cpp
+++ b/Engine/Plugins/Developer/VisualStudioSourceCodeAccess/Source/VisualStudioSourceCodeAccess/Private/VisualStudioSourceCodeAccessor.cpp
@@ -931,6 +931,13 @@ bool FVisualStudioSourceCodeAccessor::AddSourceFiles(const TArray<FString>& Abso
 								bSuccess &= true;
 							}
 						}
+
+						// Save the updated project to avoid a message when closing VS
+						CComPtr<EnvDTE::Project> Project;
+						if (SUCCEEDED(ModuleProjectFolder->get_ContainingProject(&Project)) && Project)
+						{
+							Project->Save(nullptr);
+						}
 					}
 					else
 					{
