Improve some plugin descriptions.

diff --git a/Engine/Plugins/Developer/KDevelopSourceCodeAccess/KDevelopSourceCodeAccess.uplugin b/Engine/Plugins/Developer/KDevelopSourceCodeAccess/KDevelopSourceCodeAccess.uplugin
--- a/Engine/Plugins/Developer/KDevelopSourceCodeAccess/KDevelopSourceCodeAccess.uplugin
+++ b/Engine/Plugins/Developer/KDevelopSourceCodeAccess/KDevelopSourceCodeAccess.uplugin
@@ -1,7 +1,7 @@
 {
 	"FileVersion" : 3,
 	
-	"FriendlyName" : "KDevelop Source Code Access",
+	"FriendlyName" : "KDevelop Integration",
 	"Version" : 1,
 	"VersionName" : "1.0",
 	"CreatedBy" : "Epic Games, Inc.",
diff --git a/Engine/Plugins/Developer/NullSourceCodeAccess/NullSourceCodeAccess.uplugin b/Engine/Plugins/Developer/NullSourceCodeAccess/NullSourceCodeAccess.uplugin
--- a/Engine/Plugins/Developer/NullSourceCodeAccess/NullSourceCodeAccess.uplugin
+++ b/Engine/Plugins/Developer/NullSourceCodeAccess/NullSourceCodeAccess.uplugin
@@ -1,6 +1,6 @@
 {
 	"FileVersion" : 3,
-	"FriendlyName" : "Null Source Code Access",
+    "FriendlyName": "Linux Compiler-only Integration",
 	"Version" : 1,
 	"VersionName" : "1.0",
 	"CreatedBy" : "Salamanderrake" ,
diff --git a/Engine/Plugins/Developer/VisualStudioSourceCodeAccess/VisualStudioSourceCodeAccess.uplugin b/Engine/Plugins/Developer/VisualStudioSourceCodeAccess/VisualStudioSourceCodeAccess.uplugin
--- a/Engine/Plugins/Developer/VisualStudioSourceCodeAccess/VisualStudioSourceCodeAccess.uplugin
+++ b/Engine/Plugins/Developer/VisualStudioSourceCodeAccess/VisualStudioSourceCodeAccess.uplugin
@@ -1,7 +1,7 @@
 {
 	"FileVersion" : 3,
 	
-	"FriendlyName" : "Visual Studio Source Code Access",
+	"FriendlyName" : "Visual Studio Integration",
 	"Version" : 1,
 	"VersionName" : "1.0",
 	"CreatedBy" : "Epic Games, Inc.",
diff --git a/Engine/Plugins/Developer/XCodeSourceCodeAccess/XCodeSourceCodeAccess.uplugin b/Engine/Plugins/Developer/XCodeSourceCodeAccess/XCodeSourceCodeAccess.uplugin
--- a/Engine/Plugins/Developer/XCodeSourceCodeAccess/XCodeSourceCodeAccess.uplugin
+++ b/Engine/Plugins/Developer/XCodeSourceCodeAccess/XCodeSourceCodeAccess.uplugin
@@ -1,7 +1,7 @@
 {
 	"FileVersion" : 3,
 	
-	"FriendlyName" : "XCode Source Code Access",
+	"FriendlyName" : "XCode Integration",
 	"Version" : 1,
 	"VersionName" : "1.0",
 	"CreatedBy" : "Epic Games, Inc.",
