Fixed assert when failing to enumerate a Visual Studio process

UE-12852 - Crash when immediately closing Visual Studio after opening C++ object

diff --git a/Engine/Plugins/Developer/VisualStudioSourceCodeAccess/Source/VisualStudioSourceCodeAccess/Private/VisualStudioSourceCodeAccessor.cpp b/Engine/Plugins/Developer/VisualStudioSourceCodeAccess/Source/VisualStudioSourceCodeAccess/Private/VisualStudioSourceCodeAccessor.cpp
--- a/Engine/Plugins/Developer/VisualStudioSourceCodeAccess/Source/VisualStudioSourceCodeAccess/Private/VisualStudioSourceCodeAccessor.cpp
+++ b/Engine/Plugins/Developer/VisualStudioSourceCodeAccess/Source/VisualStudioSourceCodeAccess/Private/VisualStudioSourceCodeAccessor.cpp
@@ -501,27 +501,28 @@ bool GetProcessCommandLine(const ::DWORD InProcessID, FString& OutCommandLine)
 						::ULONG uReturn = 0;
 						if (SUCCEEDED(pEnumerator->Next(WBEM_INFINITE, 1, &pclsObj, &uReturn)))
 						{
-							check(uReturn == 1);
+							if (uReturn == 1)
+							{
+								::VARIANT vtProp;
 
-							::VARIANT vtProp;
+								::DWORD CurProcessID = 0;
+								if (SUCCEEDED(pclsObj->Get(TEXT("ProcessId"), 0, &vtProp, 0, 0)))
+								{
+									CurProcessID = vtProp.ulVal;
+									::VariantClear(&vtProp);
+								}
 
-							::DWORD CurProcessID = 0;
-							if (SUCCEEDED(pclsObj->Get(TEXT("ProcessId"), 0, &vtProp, 0, 0)))
-							{
-								CurProcessID = vtProp.ulVal;
-								::VariantClear(&vtProp);
-							}
+								check(CurProcessID == InProcessID);
+								if (SUCCEEDED(pclsObj->Get(TEXT("CommandLine"), 0, &vtProp, 0, 0)))
+								{
+									OutCommandLine = vtProp.bstrVal;
+									::VariantClear(&vtProp);
 
-							check(CurProcessID == InProcessID);
-							if (SUCCEEDED(pclsObj->Get(TEXT("CommandLine"), 0, &vtProp, 0, 0)))
-							{
-								OutCommandLine = vtProp.bstrVal;
-								::VariantClear(&vtProp);
+									bSuccess = true;
+								}
 
-								bSuccess = true;
+								pclsObj->Release();
 							}
-
-							pclsObj->Release();
 						}
 					}
 
