Updating the VS source code accessor to prefer 2015, but to support 2013.

diff --git a/Engine/Plugins/Developer/VisualStudioSourceCodeAccess/Source/VisualStudioSourceCodeAccess/Private/VisualStudioSourceCodeAccessor.cpp b/Engine/Plugins/Developer/VisualStudioSourceCodeAccess/Source/VisualStudioSourceCodeAccess/Private/VisualStudioSourceCodeAccessor.cpp
--- a/Engine/Plugins/Developer/VisualStudioSourceCodeAccess/Source/VisualStudioSourceCodeAccess/Private/VisualStudioSourceCodeAccessor.cpp
+++ b/Engine/Plugins/Developer/VisualStudioSourceCodeAccess/Source/VisualStudioSourceCodeAccess/Private/VisualStudioSourceCodeAccessor.cpp
@@ -81,6 +81,7 @@ void FVisualStudioSourceCodeAccessor::Startup()
 
 	// Preferential order of VS versions
 	AddVisualStudioVersion(14); // Visual Studio 2015
+	AddVisualStudioVersion(12); // Visual Studio 2013
 }
 
 void FVisualStudioSourceCodeAccessor::Shutdown()
