Fix Visual Studio source code accessor not working correctly in editor builds created by the build system, due to VS2017 COM components not being registered.

diff --git a/Engine/Plugins/Developer/VisualStudioSourceCodeAccess/Source/VisualStudioSourceCodeAccess/Private/VisualStudioSourceCodeAccessor.cpp b/Engine/Plugins/Developer/VisualStudioSourceCodeAccess/Source/VisualStudioSourceCodeAccess/Private/VisualStudioSourceCodeAccessor.cpp
--- a/Engine/Plugins/Developer/VisualStudioSourceCodeAccess/Source/VisualStudioSourceCodeAccess/Private/VisualStudioSourceCodeAccessor.cpp
+++ b/Engine/Plugins/Developer/VisualStudioSourceCodeAccess/Source/VisualStudioSourceCodeAccess/Private/VisualStudioSourceCodeAccessor.cpp
@@ -35,7 +35,11 @@
 	#pragma warning(disable: 6244)
 
 	// import EnvDTE
-	#import "libid:80cc9f66-e7d8-4ddd-85b6-d9e6cd0e93e2" version("8.0") lcid("0") raw_interfaces_only named_guids
+	#if VSACCESSOR_HAS_DTE_OLB
+		#import "NotForLicensees/dte80a.olb" version("8.0") lcid("0") raw_interfaces_only named_guids
+	#else
+		#import "libid:80cc9f66-e7d8-4ddd-85b6-d9e6cd0e93e2" version("8.0") lcid("0") raw_interfaces_only named_guids
+	#endif
 
 	#pragma warning(pop)
 #endif
diff --git a/Engine/Plugins/Developer/VisualStudioSourceCodeAccess/Source/VisualStudioSourceCodeAccess/VisualStudioSourceCodeAccess.Build.cs b/Engine/Plugins/Developer/VisualStudioSourceCodeAccess/Source/VisualStudioSourceCodeAccess/VisualStudioSourceCodeAccess.Build.cs
--- a/Engine/Plugins/Developer/VisualStudioSourceCodeAccess/Source/VisualStudioSourceCodeAccess/VisualStudioSourceCodeAccess.Build.cs
+++ b/Engine/Plugins/Developer/VisualStudioSourceCodeAccess/Source/VisualStudioSourceCodeAccess/VisualStudioSourceCodeAccess.Build.cs
@@ -1,5 +1,6 @@
 // Copyright 1998-2018 Epic Games, Inc. All Rights Reserved.
 using Microsoft.Win32;
+using System.IO;
 
 namespace UnrealBuildTool.Rules
 {
@@ -23,24 +24,36 @@ namespace UnrealBuildTool.Rules
 			}
 
 			bool bHasVisualStudioDTE;
-			try
+
+			// In order to support building the plugin on build machines (which may not have the IDE installed), allow using an OLB rather than registered component.
+			string DteOlbPath = Path.Combine(ModuleDirectory, "Private", "NotForLicensees", "dte80a.olb");
+			if(File.Exists(DteOlbPath))
 			{
-				// Interrogate the Win32 registry
-				string DTEKey = null;
-				switch (Target.WindowsPlatform.Compiler)
-				{
-					case WindowsCompiler.VisualStudio2017:
-						DTEKey = "VisualStudio.DTE.15.0";
-						break;
-					case WindowsCompiler.VisualStudio2015:
-						DTEKey = "VisualStudio.DTE.14.0";
-						break;
-				}
-				bHasVisualStudioDTE = RegistryKey.OpenBaseKey(RegistryHive.ClassesRoot, RegistryView.Registry32).OpenSubKey(DTEKey) != null;
+				PrivateDefinitions.Add("VSACCESSOR_HAS_DTE_OLB=1");
+				bHasVisualStudioDTE = true;
 			}
-			catch
+			else
 			{
-				bHasVisualStudioDTE = false;
+				PrivateDefinitions.Add("VSACCESSOR_HAS_DTE_OLB=0");
+				try
+				{
+					// Interrogate the Win32 registry
+					string DTEKey = null;
+					switch (Target.WindowsPlatform.Compiler)
+					{
+						case WindowsCompiler.VisualStudio2017:
+							DTEKey = "VisualStudio.DTE.15.0";
+							break;
+						case WindowsCompiler.VisualStudio2015:
+							DTEKey = "VisualStudio.DTE.14.0";
+							break;
+					}
+					bHasVisualStudioDTE = RegistryKey.OpenBaseKey(RegistryHive.ClassesRoot, RegistryView.Registry32).OpenSubKey(DTEKey) != null;
+				}
+				catch
+				{
+					bHasVisualStudioDTE = false;
+				}
 			}
 
 			if (bHasVisualStudioDTE)
