Add support for Visual Studio 2019 to the source code accessor. VS2019 does not have registry entries in the same way that VS2017 do, so use VSWHERE to determine the location of any installations.

diff --git a/Engine/Plugins/Developer/VisualStudioSourceCodeAccess/Source/VisualStudioSourceCodeAccess/Private/VisualStudioSourceCodeAccessModule.cpp b/Engine/Plugins/Developer/VisualStudioSourceCodeAccess/Source/VisualStudioSourceCodeAccess/Private/VisualStudioSourceCodeAccessModule.cpp
--- a/Engine/Plugins/Developer/VisualStudioSourceCodeAccess/Source/VisualStudioSourceCodeAccess/Private/VisualStudioSourceCodeAccessModule.cpp
+++ b/Engine/Plugins/Developer/VisualStudioSourceCodeAccess/Source/VisualStudioSourceCodeAccess/Private/VisualStudioSourceCodeAccessModule.cpp
@@ -19,6 +19,7 @@ void FVisualStudioSourceCodeAccessModule::StartupModule()
 
 	// Add all the explicit version wrappers. If one of these is selected, UBT will generate project files in the appropriate format. Editor behavior is still to detect which version to use
 	// from the solution on disk.
+	RegisterWrapper("VisualStudio2019", LOCTEXT("VisualStudio2019", "Visual Studio 2019"), LOCTEXT("UsingVisualStudio2019", "Open source code files in Visual Studio 2019"));
 	RegisterWrapper("VisualStudio2017", LOCTEXT("VisualStudio2017", "Visual Studio 2017"), LOCTEXT("UsingVisualStudio2017", "Open source code files in Visual Studio 2017"));
 	RegisterWrapper("VisualStudio2015", LOCTEXT("VisualStudio2015", "Visual Studio 2015"), LOCTEXT("UsingVisualStudio2015", "Open source code files in Visual Studio 2015"));
 
diff --git a/Engine/Plugins/Developer/VisualStudioSourceCodeAccess/Source/VisualStudioSourceCodeAccess/Private/VisualStudioSourceCodeAccessor.cpp b/Engine/Plugins/Developer/VisualStudioSourceCodeAccess/Source/VisualStudioSourceCodeAccess/Private/VisualStudioSourceCodeAccessor.cpp
--- a/Engine/Plugins/Developer/VisualStudioSourceCodeAccess/Source/VisualStudioSourceCodeAccess/Private/VisualStudioSourceCodeAccessor.cpp
+++ b/Engine/Plugins/Developer/VisualStudioSourceCodeAccess/Source/VisualStudioSourceCodeAccess/Private/VisualStudioSourceCodeAccessor.cpp
@@ -15,6 +15,9 @@
 #include "ProjectDescriptor.h"
 #include "Interfaces/IProjectManager.h"
 #include "Misc/FileHelper.h"
+#include "Misc/MonitoredProcess.h"
+#include "Serialization/JsonReader.h"
+#include "Serialization/JsonSerializer.h"
 //#include "GameProjectGenerationModule.h"
 
 #if WITH_EDITOR
@@ -25,6 +28,7 @@
 #include "Windows/AllowWindowsPlatformTypes.h"
 #include "Windows/AllowWindowsPlatformAtomics.h"
 #include <unknwn.h>
+#include <ShlObj.h>
 #include "Windows/COMPointer.h"
 #if VSACCESSOR_HAS_DTE
 	#pragma warning(push)
@@ -101,7 +105,9 @@ static void OnModuleCompileStarted(bool bIsAsyncCompile)
 
 int32 GetVisualStudioVersionForCompiler()
 {
-#if _MSC_VER >= 1910
+#if _MSC_VER >= 1920
+	return 16; // Visual Studio 2019
+#elif _MSC_VER >= 1910
 	return 15; // Visual Studio 2017
 #elif _MSC_VER == 1900
 	return 14; // Visual Studio 2015
@@ -126,8 +132,16 @@ int32 GetVisualStudioVersionForSolution(const FString& InSolutionFile)
 		const int32 VersionStringStart = SolutionFileContents.Find(VisualStudioVersionString);
 		if (VersionStringStart != INDEX_NONE)
 		{
+			const TCHAR* VersionChar = *SolutionFileContents + VersionStringStart + VisualStudioVersionString.Len();
+
+			const TCHAR VersionSuffix[] = TEXT("Version ");
+			if (FCString::Strnicmp(VersionChar, VersionSuffix, ARRAY_COUNT(VersionSuffix) - 1) == 0)
+			{
+				VersionChar += ARRAY_COUNT(VersionSuffix) - 1;
+			}
+
 			FString VersionString;
-			for (const TCHAR* VersionChar = *SolutionFileContents + VersionStringStart + VisualStudioVersionString.Len(); FChar::IsDigit(*VersionChar); ++VersionChar)
+			for (; FChar::IsDigit(*VersionChar); ++VersionChar)
 			{
 				VersionString.AppendChar(*VersionChar);
 			}
@@ -160,6 +174,7 @@ void FVisualStudioSourceCodeAccessor::RefreshAvailability()
 {
 	Locations.Reset();
 
+	AddVisualStudioVersionUsingVsWhere(16); // Visual Studio 2019
 	AddVisualStudioVersion(15); // Visual Studio 2017
 	AddVisualStudioVersion(14); // Visual Studio 2015
 	AddVisualStudioVersion(12); // Visual Studio 2013
@@ -1176,6 +1191,62 @@ void FVisualStudioSourceCodeAccessor::AddVisualStudioVersion(const int MajorVers
 	}
 }
 
+void FVisualStudioSourceCodeAccessor::AddVisualStudioVersionUsingVsWhere(int VersionNumber)
+{
+	TCHAR* ProgramFilesPath = nullptr;
+	if (FAILED(SHGetKnownFolderPath(FOLDERID_ProgramFilesX86, KF_FLAG_DEFAULT, NULL, &ProgramFilesPath)))
+	{
+		UE_LOG(LogVSAccessor, Warning, TEXT("Unable to get ProgramFiles(x86) path"));
+		return;
+	}
+
+	FString VsWhereLocation = FString(ProgramFilesPath) / TEXT("Microsoft Visual Studio/Installer/vswhere.exe");
+	CoTaskMemFree(ProgramFilesPath);
+
+	if (!FPaths::FileExists(VsWhereLocation))
+	{
+		UE_LOG(LogVSAccessor, Display, TEXT("Unable to find VSWHERE installation"));
+		return;
+	}
+
+	FString ProcessOutput;
+
+	FMonitoredProcess Process(VsWhereLocation, *FString::Printf(TEXT("-version [%d.0,%d.0] -format json -prerelease"), VersionNumber, VersionNumber + 1), true);
+	Process.OnOutput().BindLambda([&ProcessOutput](const FString& Str) { ProcessOutput += Str; });
+	Process.Launch();
+	while (Process.Update())
+	{
+		FPlatformProcess::Sleep(0.01);
+	}
+
+	TArray<TSharedPtr<FJsonValue>> Array;
+	TSharedRef<TJsonReader<>> Reader = TJsonReaderFactory<>::Create(ProcessOutput);
+	if (!FJsonSerializer::Deserialize(Reader, Array))
+	{
+		UE_LOG(LogVSAccessor, Warning, TEXT("Unable to parse output from VSWHERE: %s"), *ProcessOutput);
+		return;
+	}
+
+	for (const TSharedPtr<FJsonValue>& Item : Array)
+	{
+		const TSharedPtr<FJsonObject>* Object;
+		if (Item->TryGetObject(Object))
+		{
+			FString ProductPath;
+			if ((*Object)->TryGetStringField(TEXT("productPath"), ProductPath))
+			{
+				VisualStudioLocation NewLocation;
+				NewLocation.VersionNumber = VersionNumber;
+				NewLocation.ExecutablePath = ProductPath;
+#if VSACCESSOR_HAS_DTE
+				NewLocation.ROTMoniker = FString::Printf(TEXT("!VisualStudio.DTE.%d.0"), VersionNumber);
+#endif
+				Locations.Add(NewLocation);
+			}
+		}
+	}
+}
+
 TArray<FVisualStudioSourceCodeAccessor::VisualStudioLocation> FVisualStudioSourceCodeAccessor::GetPrioritizedVisualStudioVersions(const FString& InSolution) const
 {
 	TArray<VisualStudioLocation> PrioritizedLocations = Locations;
diff --git a/Engine/Plugins/Developer/VisualStudioSourceCodeAccess/Source/VisualStudioSourceCodeAccess/Private/VisualStudioSourceCodeAccessor.h b/Engine/Plugins/Developer/VisualStudioSourceCodeAccess/Source/VisualStudioSourceCodeAccess/Private/VisualStudioSourceCodeAccessor.h
--- a/Engine/Plugins/Developer/VisualStudioSourceCodeAccess/Source/VisualStudioSourceCodeAccess/Private/VisualStudioSourceCodeAccessor.h
+++ b/Engine/Plugins/Developer/VisualStudioSourceCodeAccess/Source/VisualStudioSourceCodeAccess/Private/VisualStudioSourceCodeAccessor.h
@@ -100,6 +100,13 @@ private:
 	 */
 	void AddVisualStudioVersion(const int MajorVersion, const bool bAllowExpress = true);
 
+	/**
+	 * Add a new version of Visual Studio to the supported locations array using VSWHERE
+	 *
+	 * @param	MajorVersion	The major version number of Visual Studio (eg, 11 for VS2012, 12 for VS2013)
+	 */
+	void AddVisualStudioVersionUsingVsWhere(const int MajorVersion);
+
 	/** 
 	 * Get the prioritized list of VS install locations based upon the version of the given solution
 	 * 
diff --git a/Engine/Plugins/Developer/VisualStudioSourceCodeAccess/Source/VisualStudioSourceCodeAccess/VisualStudioSourceCodeAccess.Build.cs b/Engine/Plugins/Developer/VisualStudioSourceCodeAccess/Source/VisualStudioSourceCodeAccess/VisualStudioSourceCodeAccess.Build.cs
--- a/Engine/Plugins/Developer/VisualStudioSourceCodeAccess/Source/VisualStudioSourceCodeAccess/VisualStudioSourceCodeAccess.Build.cs
+++ b/Engine/Plugins/Developer/VisualStudioSourceCodeAccess/Source/VisualStudioSourceCodeAccess/VisualStudioSourceCodeAccess.Build.cs
@@ -14,7 +14,8 @@ namespace UnrealBuildTool.Rules
 					"Core",
 					"SourceCodeAccess",
 					"DesktopPlatform",
-					"Projects"
+					"Projects",
+					"Json"
 				}
 			);
 
diff --git a/Engine/Source/Developer/Windows/WindowsTargetPlatform/Classes/WindowsTargetSettings.h b/Engine/Source/Developer/Windows/WindowsTargetPlatform/Classes/WindowsTargetSettings.h
--- a/Engine/Source/Developer/Windows/WindowsTargetPlatform/Classes/WindowsTargetSettings.h
+++ b/Engine/Source/Developer/Windows/WindowsTargetPlatform/Classes/WindowsTargetSettings.h
@@ -22,6 +22,7 @@ enum class ECompilerVersion : uint8
 	Default = 0,
 	VisualStudio2015 = 1 UMETA(DisplayName = "Visual Studio 2015"),
 	VisualStudio2017 = 2 UMETA(DisplayName = "Visual Studio 2017"),
+	VisualStudio2019 = 3 UMETA(DisplayName = "Visual Studio 2019"),
 };
 
 /**
