Fix partial path to versions of Visual Studio enumerated via the Setup API, causing creation of new projects to fail. Also change ordering of enumerated instances to prefer non-preview editions.

diff --git a/Engine/Plugins/Developer/VisualStudioSourceCodeAccess/Source/VisualStudioSourceCodeAccess/Private/VisualStudioSourceCodeAccessor.cpp b/Engine/Plugins/Developer/VisualStudioSourceCodeAccess/Source/VisualStudioSourceCodeAccess/Private/VisualStudioSourceCodeAccessor.cpp
--- a/Engine/Plugins/Developer/VisualStudioSourceCodeAccess/Source/VisualStudioSourceCodeAccess/Private/VisualStudioSourceCodeAccessor.cpp
+++ b/Engine/Plugins/Developer/VisualStudioSourceCodeAccess/Source/VisualStudioSourceCodeAccess/Private/VisualStudioSourceCodeAccessor.cpp
@@ -1163,6 +1163,7 @@ void FVisualStudioSourceCodeAccessor::AddVisualStudioVersion(const int MajorVers
 
 	VisualStudioLocation NewLocation;
 	NewLocation.VersionNumber = MajorVersion;
+	NewLocation.bPreviewRelease = false;
 	NewLocation.ExecutablePath = BaseExecutablePath / TEXT("devenv.exe");
 #if VSACCESSOR_HAS_DTE
 	NewLocation.ROTMoniker = FString::Printf(TEXT("!VisualStudio.DTE.%d.0"), MajorVersion);
@@ -1241,19 +1242,39 @@ void FVisualStudioSourceCodeAccessor::AddVisualStudioVersionUsingVisualStudioSet
 				{
 					if (FCString::Strncmp(InstallationVersion, VersionPrefix, VersionPrefixLen) == 0)
 					{
-						BSTR ProductPath;
-						Result = Instance2->GetProductPath(&ProductPath);
+						BSTR InstallationPath;
+						Result = Instance2->GetInstallationPath(&InstallationPath);
 						if (SUCCEEDED(Result))
 						{
-							VisualStudioLocation NewLocation;
-							NewLocation.VersionNumber = VersionNumber;
-							NewLocation.ExecutablePath = ProductPath;
+							BSTR ProductPath;
+							Result = Instance2->GetProductPath(&ProductPath);
+							if (SUCCEEDED(Result))
+							{
+								VisualStudioLocation NewLocation;
+								NewLocation.VersionNumber = VersionNumber;
+								NewLocation.bPreviewRelease = false;
+								NewLocation.ExecutablePath = FString::Printf(TEXT("%s\\%s"), InstallationPath, ProductPath);
 #if VSACCESSOR_HAS_DTE
-							NewLocation.ROTMoniker = FString::Printf(TEXT("!VisualStudio.DTE.%d.0"), VersionNumber);
+								NewLocation.ROTMoniker = FString::Printf(TEXT("!VisualStudio.DTE.%d.0"), VersionNumber);
 #endif
-							Locations.Add(NewLocation);
 
-							SysFreeString(ProductPath);
+								TComPtr<ISetupInstanceCatalog> Catalog;
+								Result = Instance2->QueryInterface(__uuidof(ISetupInstanceCatalog), (LPVOID*)&Catalog);
+								if (SUCCEEDED(Result) && Catalog != nullptr)
+								{
+									VARIANT_BOOL PrereleaseFlag;
+									Result = Catalog->IsPrerelease(&PrereleaseFlag);
+									if (SUCCEEDED(Result) && PrereleaseFlag == VARIANT_TRUE)
+									{
+										NewLocation.bPreviewRelease = true;
+									}
+								}
+
+								Locations.Add(NewLocation);
+
+								SysFreeString(ProductPath);
+							}
+							SysFreeString(InstallationPath);
 						}
 					}
 					SysFreeString(InstallationVersion);
@@ -1277,8 +1298,8 @@ TArray<FVisualStudioSourceCodeAccessor::VisualStudioLocation> FVisualStudioSourc
 	{
 		PrioritizedLocations.StableSort([&](const VisualStudioLocation& InFirst, const VisualStudioLocation& InSecond) -> bool
 		{
-			const int32 FirstSortWeight = (InFirst.VersionNumber == SolutionVersion) ? 1 : 0;
-			const int32 SecondSortWeight = (InSecond.VersionNumber == SolutionVersion) ? 1 : 0;
+			const int32 FirstSortWeight = (InFirst.VersionNumber == SolutionVersion) ? (InFirst.bPreviewRelease? 1 : 2) : 0;
+			const int32 SecondSortWeight = (InSecond.VersionNumber == SolutionVersion) ? (InSecond.bPreviewRelease? 1 : 2) : 0;
 			return FirstSortWeight >= SecondSortWeight;
 		});
 	}
diff --git a/Engine/Plugins/Developer/VisualStudioSourceCodeAccess/Source/VisualStudioSourceCodeAccess/Private/VisualStudioSourceCodeAccessor.h b/Engine/Plugins/Developer/VisualStudioSourceCodeAccess/Source/VisualStudioSourceCodeAccess/Private/VisualStudioSourceCodeAccessor.h
--- a/Engine/Plugins/Developer/VisualStudioSourceCodeAccess/Source/VisualStudioSourceCodeAccess/Private/VisualStudioSourceCodeAccessor.h
+++ b/Engine/Plugins/Developer/VisualStudioSourceCodeAccess/Source/VisualStudioSourceCodeAccess/Private/VisualStudioSourceCodeAccessor.h
@@ -10,6 +10,7 @@ public:
 	struct VisualStudioLocation
 	{
 		int32 VersionNumber;
+		bool bPreviewRelease;
 		FString ExecutablePath;
 #if VSACCESSOR_HAS_DTE
 		FString ROTMoniker;
