Fix inability to open Visual Studio 2019 after creating a project.

diff --git a/Engine/Plugins/Developer/VisualStudioSourceCodeAccess/Source/VisualStudioSourceCodeAccess/VisualStudioSourceCodeAccess.Build.cs b/Engine/Plugins/Developer/VisualStudioSourceCodeAccess/Source/VisualStudioSourceCodeAccess/VisualStudioSourceCodeAccess.Build.cs
--- a/Engine/Plugins/Developer/VisualStudioSourceCodeAccess/Source/VisualStudioSourceCodeAccess/VisualStudioSourceCodeAccess.Build.cs
+++ b/Engine/Plugins/Developer/VisualStudioSourceCodeAccess/Source/VisualStudioSourceCodeAccess/VisualStudioSourceCodeAccess.Build.cs
@@ -43,10 +43,13 @@ namespace UnrealBuildTool.Rules
 					string DTEKey = null;
 					switch (Target.WindowsPlatform.Compiler)
 					{
-						case WindowsCompiler.VisualStudio2017:
-							DTEKey = "VisualStudio.DTE.15.0";
-							break;
-						case WindowsCompiler.VisualStudio2015_DEPRECATED:
+                        case WindowsCompiler.VisualStudio2019:
+                            DTEKey = "VisualStudio.DTE.16.0";
+                            break;
+                        case WindowsCompiler.VisualStudio2017:
+                            DTEKey = "VisualStudio.DTE.15.0";
+                            break;
+                        case WindowsCompiler.VisualStudio2015_DEPRECATED:
 							DTEKey = "VisualStudio.DTE.14.0";
 							break;
 					}
