Back out revision 11 from //UE4/Dev-Core/Engine/Plugins/Developer/VisualStudioSourceCodeAccess/Source/VisualStudioSourceCodeAccess/VisualStudioSourceCodeAccess.Build.cs

diff --git a/Engine/Plugins/Developer/VisualStudioSourceCodeAccess/Source/VisualStudioSourceCodeAccess/VisualStudioSourceCodeAccess.Build.cs b/Engine/Plugins/Developer/VisualStudioSourceCodeAccess/Source/VisualStudioSourceCodeAccess/VisualStudioSourceCodeAccess.Build.cs
--- a/Engine/Plugins/Developer/VisualStudioSourceCodeAccess/Source/VisualStudioSourceCodeAccess/VisualStudioSourceCodeAccess.Build.cs
+++ b/Engine/Plugins/Developer/VisualStudioSourceCodeAccess/Source/VisualStudioSourceCodeAccess/VisualStudioSourceCodeAccess.Build.cs
@@ -1,6 +1,5 @@
 // Copyright 1998-2018 Epic Games, Inc. All Rights Reserved.
 using Microsoft.Win32;
-using System.IO;
 
 namespace UnrealBuildTool.Rules
 {
@@ -24,36 +23,24 @@ namespace UnrealBuildTool.Rules
 			}
 
 			bool bHasVisualStudioDTE;
-
-			// In order to support building the plugin on build machines (which may not have the IDE installed), allow using an OLB rather than registered component.
-			string DteOlbPath = Path.Combine(ModuleDirectory, "Private", "NotForLicensees", "dte80a.olb");
-			if(File.Exists(DteOlbPath))
-			{
-				PrivateDefinitions.Add("VSACCESSOR_HAS_DTE_OLB=1");
-				bHasVisualStudioDTE = true;
-			}
-			else
+			try
 			{
-				PrivateDefinitions.Add("VSACCESSOR_HAS_DTE_OLB=0");
-				try
+				// Interrogate the Win32 registry
+				string DTEKey = null;
+				switch (Target.WindowsPlatform.Compiler)
 				{
-					// Interrogate the Win32 registry
-					string DTEKey = null;
-					switch (Target.WindowsPlatform.Compiler)
-					{
-						case WindowsCompiler.VisualStudio2017:
-							DTEKey = "VisualStudio.DTE.15.0";
-							break;
-						case WindowsCompiler.VisualStudio2015:
-							DTEKey = "VisualStudio.DTE.14.0";
-							break;
-					}
-					bHasVisualStudioDTE = RegistryKey.OpenBaseKey(RegistryHive.ClassesRoot, RegistryView.Registry32).OpenSubKey(DTEKey) != null;
-				}
-				catch
-				{
-					bHasVisualStudioDTE = false;
+					case WindowsCompiler.VisualStudio2017:
+						DTEKey = "VisualStudio.DTE.15.0";
+						break;
+					case WindowsCompiler.VisualStudio2015:
+						DTEKey = "VisualStudio.DTE.14.0";
+						break;
 				}
+				bHasVisualStudioDTE = RegistryKey.OpenBaseKey(RegistryHive.ClassesRoot, RegistryView.Registry32).OpenSubKey(DTEKey) != null;
+			}
+			catch
+			{
+				bHasVisualStudioDTE = false;
 			}
 
 			if (bHasVisualStudioDTE)
