Don't use #import <TypeLib> for VS source code accessor when building with Clang; it's not supported.

#rnx

diff --git a/Engine/Plugins/Developer/VisualStudioSourceCodeAccess/Source/VisualStudioSourceCodeAccess/VisualStudioSourceCodeAccess.Build.cs b/Engine/Plugins/Developer/VisualStudioSourceCodeAccess/Source/VisualStudioSourceCodeAccess/VisualStudioSourceCodeAccess.Build.cs
--- a/Engine/Plugins/Developer/VisualStudioSourceCodeAccess/Source/VisualStudioSourceCodeAccess/VisualStudioSourceCodeAccess.Build.cs
+++ b/Engine/Plugins/Developer/VisualStudioSourceCodeAccess/Source/VisualStudioSourceCodeAccess/VisualStudioSourceCodeAccess.Build.cs
@@ -27,7 +27,7 @@ namespace UnrealBuildTool.Rules
 
 			// In order to support building the plugin on build machines (which may not have the IDE installed), allow using an OLB rather than registered component.
 			string DteOlbPath = Path.Combine(ModuleDirectory, "Private", "NotForLicensees", "dte80a.olb");
-			if(File.Exists(DteOlbPath))
+			if(File.Exists(DteOlbPath) && Target.WindowsPlatform.Compiler != WindowsCompiler.Clang)
 			{
 				PrivateDefinitions.Add("VSACCESSOR_HAS_DTE_OLB=1");
 				bHasVisualStudioDTE = true;
