Whitelist plugins that need to be loaded by UFE.

diff --git a/Engine/Plugins/Developer/VisualStudioSourceCodeAccess/VisualStudioSourceCodeAccess.uplugin b/Engine/Plugins/Developer/VisualStudioSourceCodeAccess/VisualStudioSourceCodeAccess.uplugin
--- a/Engine/Plugins/Developer/VisualStudioSourceCodeAccess/VisualStudioSourceCodeAccess.uplugin
+++ b/Engine/Plugins/Developer/VisualStudioSourceCodeAccess/VisualStudioSourceCodeAccess.uplugin
@@ -14,6 +14,7 @@
 	"CanContainContent" : false,
 	"IsBetaVersion" : false,
 	"Installed" : false,
+	"SupportedPrograms": [ "UnrealFrontend" ],
 	"Modules" :
 	[
 		{
@@ -24,7 +25,8 @@
 			[
 				"Win64",
 				"Win32"
-			]
+			],
+			"WhitelistPrograms" : [ "UnrealFrontend" ]
 		}
 	]
 }
\ No newline at end of file
diff --git a/Engine/Plugins/Developer/XCodeSourceCodeAccess/XCodeSourceCodeAccess.uplugin b/Engine/Plugins/Developer/XCodeSourceCodeAccess/XCodeSourceCodeAccess.uplugin
--- a/Engine/Plugins/Developer/XCodeSourceCodeAccess/XCodeSourceCodeAccess.uplugin
+++ b/Engine/Plugins/Developer/XCodeSourceCodeAccess/XCodeSourceCodeAccess.uplugin
@@ -14,6 +14,7 @@
 	"CanContainContent" : false,
 	"IsBetaVersion" : false,
 	"Installed" : false,
+	"SupportedPrograms": [ "UnrealFrontend" ],
 	"Modules" :
 	[
 		{
@@ -23,7 +24,8 @@
 			"WhitelistPlatforms" :
 			[
 				"Mac"
-			]
+			],
+			"WhitelistPrograms" : [ "UnrealFrontend" ]
 		}
 	]
 }
\ No newline at end of file
diff --git a/Engine/Plugins/Messaging/TcpMessaging/TcpMessaging.uplugin b/Engine/Plugins/Messaging/TcpMessaging/TcpMessaging.uplugin
--- a/Engine/Plugins/Messaging/TcpMessaging/TcpMessaging.uplugin
+++ b/Engine/Plugins/Messaging/TcpMessaging/TcpMessaging.uplugin
@@ -14,6 +14,7 @@
 	"CanContainContent" : false,
 	"IsBetaVersion" : false,
 	"Installed" : false,
+	"SupportedPrograms": [ "UnrealFrontend" ],
 	"Modules" :
 	[
 		{
diff --git a/Engine/Plugins/Messaging/UdpMessaging/UdpMessaging.uplugin b/Engine/Plugins/Messaging/UdpMessaging/UdpMessaging.uplugin
--- a/Engine/Plugins/Messaging/UdpMessaging/UdpMessaging.uplugin
+++ b/Engine/Plugins/Messaging/UdpMessaging/UdpMessaging.uplugin
@@ -14,6 +14,7 @@
 	"CanContainContent" : false,
 	"IsBetaVersion" : false,
 	"Installed" : false,
+	"SupportedPrograms" : [ "UnrealFrontend" ],
 	"Modules" :
 	[
 		{
