Fix "Loading Visual Studio" popup not closing when a specific Visual Studio version is selected.

diff --git a/Engine/Plugins/Developer/VisualStudioSourceCodeAccess/Source/VisualStudioSourceCodeAccess/Private/VisualStudioSourceCodeAccessorWrapper.cpp b/Engine/Plugins/Developer/VisualStudioSourceCodeAccess/Source/VisualStudioSourceCodeAccess/Private/VisualStudioSourceCodeAccessorWrapper.cpp
--- a/Engine/Plugins/Developer/VisualStudioSourceCodeAccess/Source/VisualStudioSourceCodeAccess/Private/VisualStudioSourceCodeAccessorWrapper.cpp
+++ b/Engine/Plugins/Developer/VisualStudioSourceCodeAccess/Source/VisualStudioSourceCodeAccess/Private/VisualStudioSourceCodeAccessorWrapper.cpp
@@ -76,5 +76,5 @@ bool FVisualStudioSourceCodeAccessorWrapper::SaveAllOpenDocuments() const
 
 void FVisualStudioSourceCodeAccessorWrapper::Tick(const float DeltaTime)
 {
-	// Inner will get this call separately
+	Inner->Tick(DeltaTime);
 }
