Convert file paths to use backslashes when opening source files in Visual Studio.
Fixes issues with plugins in Visual Studio prompting to re-open the file.

diff --git a/Engine/Plugins/Developer/VisualStudioSourceCodeAccess/Source/VisualStudioSourceCodeAccess/Private/VisualStudioSourceCodeAccessor.cpp b/Engine/Plugins/Developer/VisualStudioSourceCodeAccess/Source/VisualStudioSourceCodeAccess/Private/VisualStudioSourceCodeAccessor.cpp
--- a/Engine/Plugins/Developer/VisualStudioSourceCodeAccess/Source/VisualStudioSourceCodeAccess/Private/VisualStudioSourceCodeAccessor.cpp
+++ b/Engine/Plugins/Developer/VisualStudioSourceCodeAccess/Source/VisualStudioSourceCodeAccess/Private/VisualStudioSourceCodeAccessor.cpp
@@ -374,7 +374,8 @@ bool FVisualStudioSourceCodeAccessor::OpenVisualStudioFilesInternalViaDTE(const
 						}
 
 						// Open File
-						auto ANSIPath = StringCast<ANSICHAR>(*Request.FullPath);
+						FString OSFriendlyPath = Request.FullPath.Replace(TEXT("/"), TEXT("\\"), ESearchCase::CaseSensitive);
+						auto ANSIPath = StringCast<ANSICHAR>(*OSFriendlyPath);
 						FComBSTR COMStrFileName(ANSIPath.Get());
 						FComBSTR COMStrKind(EnvDTE::vsViewKindTextView);
 						TComPtr<EnvDTE::Window> Window;
@@ -1322,7 +1323,7 @@ bool FVisualStudioSourceCodeAccessor::RunVisualStudioAndOpenSolutionAndFiles(con
 			if (FPaths::FileExists(Request.FullPath))
 			{
 				Params += TEXT(" \"");
-				Params += Request.FullPath;
+				Params += Request.FullPath.Replace(TEXT("/"), TEXT("\\"), ESearchCase::CaseSensitive);
 				Params += TEXT("\"");
 
 				GoToLine = Request.LineNumber;
