Updated open source code file in Editor to call MakePlatformFilename for Visual Studio.
#rnx

diff --git a/Engine/Plugins/Developer/VisualStudioSourceCodeAccess/Source/VisualStudioSourceCodeAccess/Private/VisualStudioSourceCodeAccessor.cpp b/Engine/Plugins/Developer/VisualStudioSourceCodeAccess/Source/VisualStudioSourceCodeAccess/Private/VisualStudioSourceCodeAccessor.cpp
--- a/Engine/Plugins/Developer/VisualStudioSourceCodeAccess/Source/VisualStudioSourceCodeAccess/Private/VisualStudioSourceCodeAccessor.cpp
+++ b/Engine/Plugins/Developer/VisualStudioSourceCodeAccess/Source/VisualStudioSourceCodeAccess/Private/VisualStudioSourceCodeAccessor.cpp
@@ -374,8 +374,9 @@ bool FVisualStudioSourceCodeAccessor::OpenVisualStudioFilesInternalViaDTE(const
 						}
 
 						// Open File
-						FString OSFriendlyPath = Request.FullPath.Replace(TEXT("/"), TEXT("\\"), ESearchCase::CaseSensitive);
-						auto ANSIPath = StringCast<ANSICHAR>(*OSFriendlyPath);
+						FString PlatformFilename = Request.FullPath;
+						FPaths::MakePlatformFilename(PlatformFilename);
+						auto ANSIPath = StringCast<ANSICHAR>(*PlatformFilename);
 						FComBSTR COMStrFileName(ANSIPath.Get());
 						FComBSTR COMStrKind(EnvDTE::vsViewKindTextView);
 						TComPtr<EnvDTE::Window> Window;
@@ -1323,7 +1324,9 @@ bool FVisualStudioSourceCodeAccessor::RunVisualStudioAndOpenSolutionAndFiles(con
 			if (FPaths::FileExists(Request.FullPath))
 			{
 				Params += TEXT(" \"");
-				Params += Request.FullPath.Replace(TEXT("/"), TEXT("\\"), ESearchCase::CaseSensitive);
+				FString PlatformFilename = Request.FullPath;
+				FPaths::MakePlatformFilename(PlatformFilename);
+				Params += PlatformFilename;
 				Params += TEXT("\"");
 
 				GoToLine = Request.LineNumber;
