Fixed VS accessor to make sure source files are placed in the precise module that was selected and not just a module beginning with the same name

diff --git a/Engine/Plugins/Developer/VisualStudioSourceCodeAccess/Source/VisualStudioSourceCodeAccess/Private/VisualStudioSourceCodeAccessor.cpp b/Engine/Plugins/Developer/VisualStudioSourceCodeAccess/Source/VisualStudioSourceCodeAccess/Private/VisualStudioSourceCodeAccessor.cpp
--- a/Engine/Plugins/Developer/VisualStudioSourceCodeAccess/Source/VisualStudioSourceCodeAccess/Private/VisualStudioSourceCodeAccessor.cpp
+++ b/Engine/Plugins/Developer/VisualStudioSourceCodeAccess/Source/VisualStudioSourceCodeAccess/Private/VisualStudioSourceCodeAccessor.cpp
@@ -953,7 +953,7 @@ bool FVisualStudioSourceCodeAccessor::AddSourceFiles(const TArray<FString>& Abso
 		for (const FString& SourceFile : AbsoluteSourcePaths)
 		{
 			// First check to see if this source file is in the same module as the last source file - this is usually the case, and saves us a lot of string compares
-			if (LastSourceFilesModule && SourceFile.StartsWith(LastSourceFilesModule->ModulePath))
+			if (LastSourceFilesModule && SourceFile.StartsWith(LastSourceFilesModule->ModulePath + "/"))
 			{
 				FModuleNewSourceFiles& ModuleNewSourceFiles = ModuleToNewSourceFiles.FindChecked(LastSourceFilesModule->ModuleName);
 				ModuleNewSourceFiles.NewSourceFiles.Add(SourceFile);
@@ -964,7 +964,7 @@ bool FVisualStudioSourceCodeAccessor::AddSourceFiles(const TArray<FString>& Abso
 			LastSourceFilesModule = nullptr;
 			for (const FModuleNameAndPath& ModuleNameAndPath : ModuleNamesAndPaths)
 			{
-				if (SourceFile.StartsWith(ModuleNameAndPath.ModulePath))
+				if (SourceFile.StartsWith(ModuleNameAndPath.ModulePath + "/"))
 				{
 					LastSourceFilesModule = &ModuleNameAndPath;
 
