Jira: Seeing a 1-2 second delay when opening File Menu in editor

Use FUProjectDictionary singleton to eliminate needless Refresh() inside DoesSolutionExist()/GetSolutionPath() queries.

The code accessors want to discover the active game project solution/project, which should exist at editor launch time. We don't need the more recent refreshes from the temporary stack instances of FUProjectDictionary.

diff --git a/Engine/Plugins/Developer/CodeLiteSourceCodeAccess/Source/CodeLiteSourceCodeAccess/Private/CodeLiteSourceCodeAccessor.cpp b/Engine/Plugins/Developer/CodeLiteSourceCodeAccess/Source/CodeLiteSourceCodeAccess/Private/CodeLiteSourceCodeAccessor.cpp
--- a/Engine/Plugins/Developer/CodeLiteSourceCodeAccess/Source/CodeLiteSourceCodeAccess/Private/CodeLiteSourceCodeAccessor.cpp
+++ b/Engine/Plugins/Developer/CodeLiteSourceCodeAccess/Source/CodeLiteSourceCodeAccess/Private/CodeLiteSourceCodeAccessor.cpp
@@ -176,7 +176,7 @@ FString FCodeLiteSourceCodeAccessor::GetSolutionPath() const
 	{
 		CachedSolutionPath = FPaths::ConvertRelativePathToFull(FPaths::ProjectDir());
 		
-		if (!FUProjectDictionary(FPaths::RootDir()).IsForeignProject(CachedSolutionPath))
+		if (!FUProjectDictionary::GetDefault().IsForeignProject(CachedSolutionPath))
 		{
 			CachedSolutionPath = FPaths::Combine(FPaths::RootDir(), TEXT("UE4.workspace"));
 		}
diff --git a/Engine/Plugins/Developer/KDevelopSourceCodeAccess/Source/KDevelopSourceCodeAccess/Private/KDevelopSourceCodeAccessor.cpp b/Engine/Plugins/Developer/KDevelopSourceCodeAccess/Source/KDevelopSourceCodeAccess/Private/KDevelopSourceCodeAccessor.cpp
--- a/Engine/Plugins/Developer/KDevelopSourceCodeAccess/Source/KDevelopSourceCodeAccess/Private/KDevelopSourceCodeAccessor.cpp
+++ b/Engine/Plugins/Developer/KDevelopSourceCodeAccess/Source/KDevelopSourceCodeAccess/Private/KDevelopSourceCodeAccessor.cpp
@@ -161,7 +161,7 @@ FString FKDevelopSourceCodeAccessor::GetSolutionPath() const
 	{
 		CachedSolutionPath = FPaths::ProjectDir();
 		
-		if (!FUProjectDictionary(FPaths::RootDir()).IsForeignProject(CachedSolutionPath))
+		if (!FUProjectDictionaryFUProjectDictionary::GetDefault().IsForeignProject(CachedSolutionPath))
 		{
 			CachedSolutionPath = FPaths::Combine(FPaths::RootDir(), TEXT("UE4.kdev4"));
 		}
diff --git a/Engine/Plugins/Developer/RiderSourceCodeAccess/Source/RiderSourceCodeAccess/Private/RiderSourceCodeAccessor.cpp b/Engine/Plugins/Developer/RiderSourceCodeAccess/Source/RiderSourceCodeAccess/Private/RiderSourceCodeAccessor.cpp
--- a/Engine/Plugins/Developer/RiderSourceCodeAccess/Source/RiderSourceCodeAccess/Private/RiderSourceCodeAccessor.cpp
+++ b/Engine/Plugins/Developer/RiderSourceCodeAccess/Source/RiderSourceCodeAccess/Private/RiderSourceCodeAccessor.cpp
@@ -255,7 +255,7 @@ void FRiderSourceCodeAccessor::CachePathToSolution() const
 		{
 			CachedSolutionPath = FPaths::ConvertRelativePathToFull(FPaths::ProjectDir());
 
-			if (!FUProjectDictionary(FPaths::RootDir()).IsForeignProject(CachedSolutionPath))
+			if (!FUProjectDictionary::GetDefault().IsForeignProject(CachedSolutionPath))
 			{
 				FString MasterProjectName;
 				if (!FFileHelper::LoadFileToString(MasterProjectName, *(FPaths::EngineIntermediateDir() / TEXT("ProjectFiles/MasterProjectName.txt"))))
diff --git a/Engine/Plugins/Developer/VisualStudioCodeSourceCodeAccess/Source/VisualStudioCodeSourceCodeAccess/Private/VisualStudioCodeSourceCodeAccessor.cpp b/Engine/Plugins/Developer/VisualStudioCodeSourceCodeAccess/Source/VisualStudioCodeSourceCodeAccess/Private/VisualStudioCodeSourceCodeAccessor.cpp
--- a/Engine/Plugins/Developer/VisualStudioCodeSourceCodeAccess/Source/VisualStudioCodeSourceCodeAccess/Private/VisualStudioCodeSourceCodeAccessor.cpp
+++ b/Engine/Plugins/Developer/VisualStudioCodeSourceCodeAccess/Source/VisualStudioCodeSourceCodeAccess/Private/VisualStudioCodeSourceCodeAccessor.cpp
@@ -37,7 +37,7 @@ FString FVisualStudioCodeSourceCodeAccessor::GetSolutionPath() const
 	{
 		CachedSolutionPath = FPaths::ProjectDir();
 
-		if (!FUProjectDictionary(FPaths::RootDir()).IsForeignProject(CachedSolutionPath))
+		if (!FUProjectDictionary::GetDefault().IsForeignProject(CachedSolutionPath))
 		{
 			CachedSolutionPath = FPaths::Combine(FPaths::RootDir(), FString("UE4") + GVSCodeWorkspaceExtension);
 		}
diff --git a/Engine/Plugins/Developer/VisualStudioSourceCodeAccess/Source/VisualStudioSourceCodeAccess/Private/VisualStudioSourceCodeAccessor.cpp b/Engine/Plugins/Developer/VisualStudioSourceCodeAccess/Source/VisualStudioSourceCodeAccess/Private/VisualStudioSourceCodeAccessor.cpp
--- a/Engine/Plugins/Developer/VisualStudioSourceCodeAccess/Source/VisualStudioSourceCodeAccess/Private/VisualStudioSourceCodeAccessor.cpp
+++ b/Engine/Plugins/Developer/VisualStudioSourceCodeAccess/Source/VisualStudioSourceCodeAccess/Private/VisualStudioSourceCodeAccessor.cpp
@@ -1389,7 +1389,7 @@ FString FVisualStudioSourceCodeAccessor::GetSolutionPath() const
 		{
 			CachedSolutionPath = FPaths::ConvertRelativePathToFull(FPaths::ProjectDir());
 
-			if (!FUProjectDictionary(FPaths::RootDir()).IsForeignProject(CachedSolutionPath))
+			if (!FUProjectDictionary::GetDefault().IsForeignProject(CachedSolutionPath))
 			{
 				FString MasterProjectName;
 				if (!FFileHelper::LoadFileToString(MasterProjectName, *(FPaths::EngineIntermediateDir() / TEXT("ProjectFiles/MasterProjectName.txt"))))
diff --git a/Engine/Plugins/Developer/XCodeSourceCodeAccess/Source/XCodeSourceCodeAccess/Private/XCodeSourceCodeAccessor.cpp b/Engine/Plugins/Developer/XCodeSourceCodeAccess/Source/XCodeSourceCodeAccess/Private/XCodeSourceCodeAccessor.cpp
--- a/Engine/Plugins/Developer/XCodeSourceCodeAccess/Source/XCodeSourceCodeAccess/Private/XCodeSourceCodeAccessor.cpp
+++ b/Engine/Plugins/Developer/XCodeSourceCodeAccess/Source/XCodeSourceCodeAccess/Private/XCodeSourceCodeAccessor.cpp
@@ -93,7 +93,7 @@ FString FXCodeSourceCodeAccessor::GetSolutionPath() const
 		{
 			CachedSolutionPath = FPaths::ProjectDir();
 			
-			if (!FUProjectDictionary(FPaths::RootDir()).IsForeignProject(CachedSolutionPath))
+			if (!FUProjectDictionary::GetDefault().IsForeignProject(CachedSolutionPath))
 			{
 				CachedSolutionPath = FPaths::Combine(FPaths::RootDir(), + TEXT("UE4.xcworkspace/contents.xcworkspacedata"));
 			}
