PR #4292: UE-53233: VS GoToLine exception handling (Contributed by [REDACTED_PERSON])

diff --git a/Engine/Plugins/Developer/VisualStudioSourceCodeAccess/Source/VisualStudioSourceCodeAccess/Private/VisualStudioSourceCodeAccessor.cpp b/Engine/Plugins/Developer/VisualStudioSourceCodeAccess/Source/VisualStudioSourceCodeAccess/Private/VisualStudioSourceCodeAccessor.cpp
--- a/Engine/Plugins/Developer/VisualStudioSourceCodeAccess/Source/VisualStudioSourceCodeAccess/Private/VisualStudioSourceCodeAccessor.cpp
+++ b/Engine/Plugins/Developer/VisualStudioSourceCodeAccess/Source/VisualStudioSourceCodeAccess/Private/VisualStudioSourceCodeAccessor.cpp
@@ -390,19 +390,34 @@ bool FVisualStudioSourceCodeAccessor::OpenVisualStudioFilesInternalViaDTE(const
 							TComPtr<EnvDTE::Document> Document;
 							TComPtr<IDispatch> SelectionDispatch;
 							TComPtr<EnvDTE::TextSelection> Selection;
-							if ( SUCCEEDED(DTE->get_ActiveDocument(&Document)) &&
-								 SUCCEEDED(Document->get_Selection(&SelectionDispatch)) &&
-								 SelectionDispatch && SUCCEEDED(SelectionDispatch->QueryInterface(&Selection)) &&
-								 SUCCEEDED(Selection->GotoLine(Request.LineNumber, VARIANT_TRUE)) )
+
+							int32 RetryCount = 5;
+							do
 							{
-								if ( !SUCCEEDED(Selection->MoveToLineAndOffset(Request.LineNumber, Request.ColumnNumber, false)) )
+								// Potential exception can be thrown here in the form of 'Call was rejected by callee'
+								// See https://msdn.microsoft.com/en-us/library/ms228772.aspx
+								if (!SUCCEEDED(DTE->get_ActiveDocument(&Document)) ||
+									!SUCCEEDED(Document->get_Selection(&SelectionDispatch)) ||
+									!SUCCEEDED(SelectionDispatch->QueryInterface(&Selection)))
 								{
-									UE_LOG(LogVSAccessor, Warning, TEXT("Couldn't goto column number '%i' of line '%i' in '%s'"), Request.ColumnNumber, Request.LineNumber, *Request.FullPath);
+									if (RetryCount == 0)
+									{
+										UE_LOG(LogVSAccessor, Warning, TEXT("Couldn't goto line number '%i' in '%s'."), Request.LineNumber, *Request.FullPath);
+									}
+									else
+									{
+										FPlatformProcess::Sleep(0.1f);
+									}
 								}
-							}
-							else
+
+							} while (--RetryCount >= 0);
+
+							if (Selection.IsValid() && !SUCCEEDED(Selection->MoveToLineAndOffset(Request.LineNumber, Request.ColumnNumber, false)))
 							{
-								UE_LOG(LogVSAccessor, Warning, TEXT("Couldn't goto line number '%i' in '%s'"), Request.LineNumber, *Request.FullPath);
+								if (!SUCCEEDED(Selection->GotoLine(Request.LineNumber, true)))
+								{
+									UE_LOG(LogVSAccessor, Warning, TEXT("Couldn't goto column number '%i' of line '%i' in '%s'."), Request.ColumnNumber, Request.LineNumber, *Request.FullPath);
+								}
 							}
 						}
 						else
