Fixed open source file when creating new code class with visual studio when regenerating the project which causes a modal that prevents opening the .h and .cpp files

diff --git a/Engine/Plugins/Developer/VisualStudioSourceCodeAccess/Source/VisualStudioSourceCodeAccess/Private/VisualStudioSourceCodeAccessor.cpp b/Engine/Plugins/Developer/VisualStudioSourceCodeAccess/Source/VisualStudioSourceCodeAccess/Private/VisualStudioSourceCodeAccessor.cpp
--- a/Engine/Plugins/Developer/VisualStudioSourceCodeAccess/Source/VisualStudioSourceCodeAccess/Private/VisualStudioSourceCodeAccessor.cpp
+++ b/Engine/Plugins/Developer/VisualStudioSourceCodeAccess/Source/VisualStudioSourceCodeAccess/Private/VisualStudioSourceCodeAccessor.cpp
@@ -930,13 +930,14 @@ bool FVisualStudioSourceCodeAccessor::OpenSourceFiles(const TArray<FString>& Abs
 
 bool FVisualStudioSourceCodeAccessor::AddSourceFiles(const TArray<FString>& AbsoluteSourcePaths, const TArray<FString>& AvailableModules)
 {
+	bool bSuccess = false;
+
 	// This code is temporarily disabled because it doesn't account for UBT setting per-file properties for C++ source files,
 	// adding include paths, force-included headers, and so on. Intellisense does not work correctly without these properties being set.
 #if 0
 	// This requires DTE - there is no fallback for this operation when DTE is not available
 #if WITH_VISUALSTUDIO_DTE
-	bool bSuccess = true;
-
+	bSuccess = true;
 	struct FModuleNameAndPath
 	{
 		FString ModuleBuildFilePath;
@@ -1117,12 +1118,14 @@ bool FVisualStudioSourceCodeAccessor::AddSourceFiles(const TArray<FString>& Abso
 		UE_LOG(LogVSAccessor, Verbose, TEXT("Cannot add source files as Visual Studio is either not open or not responding"));
 		bSuccess = false;
 	}
-
-	return bSuccess;
 #endif
+#else
+	// if we add new source files but do not add them directly and rely on project generation instead, if we have an opened instance, request for the files to be opened now
+	// this is because project generation will trigger a modal on our opened instance and prevent file open request to be handled.
+	OpenSourceFiles(AbsoluteSourcePaths);
 #endif
 
-	return false;
+	return bSuccess;
 }
 
 bool FVisualStudioSourceCodeAccessor::OpenFileAtLine(const FString& FullPath, int32 LineNumber, int32 ColumnNumber)
