Support VisualStudio2022 as a source code accessor and compilier
Note: _MSC_VER is still 1929 with the VS2022 preview, I expect it to be 1930 eventually when the compiler version is bumped which should match the versioning scheme for every other VS release

diff --git a/Engine/Plugins/Developer/VisualStudioSourceCodeAccess/Source/VisualStudioSourceCodeAccess/Private/VisualStudioSourceCodeAccessModule.cpp b/Engine/Plugins/Developer/VisualStudioSourceCodeAccess/Source/VisualStudioSourceCodeAccess/Private/VisualStudioSourceCodeAccessModule.cpp
--- a/Engine/Plugins/Developer/VisualStudioSourceCodeAccess/Source/VisualStudioSourceCodeAccess/Private/VisualStudioSourceCodeAccessModule.cpp
+++ b/Engine/Plugins/Developer/VisualStudioSourceCodeAccess/Source/VisualStudioSourceCodeAccess/Private/VisualStudioSourceCodeAccessModule.cpp
@@ -19,6 +19,7 @@ void FVisualStudioSourceCodeAccessModule::StartupModule()
 
 	// Add all the explicit version wrappers. If one of these is selected, UBT will generate project files in the appropriate format. Editor behavior is still to detect which version to use
 	// from the solution on disk.
+	RegisterWrapper("VisualStudio2022", LOCTEXT("VisualStudio2022", "Visual Studio 2022"), LOCTEXT("UsingVisualStudio2022", "Open source code files in Visual Studio 2022"));
 	RegisterWrapper("VisualStudio2019", LOCTEXT("VisualStudio2019", "Visual Studio 2019"), LOCTEXT("UsingVisualStudio2019", "Open source code files in Visual Studio 2019"));
 	RegisterWrapper("VisualStudio2017", LOCTEXT("VisualStudio2017", "Visual Studio 2017"), LOCTEXT("UsingVisualStudio2017", "Open source code files in Visual Studio 2017"));
 	RegisterWrapper("VisualStudio2015", LOCTEXT("VisualStudio2015", "Visual Studio 2015"), LOCTEXT("UsingVisualStudio2015", "Open source code files in Visual Studio 2015"));
diff --git a/Engine/Plugins/Developer/VisualStudioSourceCodeAccess/Source/VisualStudioSourceCodeAccess/Private/VisualStudioSourceCodeAccessor.cpp b/Engine/Plugins/Developer/VisualStudioSourceCodeAccess/Source/VisualStudioSourceCodeAccess/Private/VisualStudioSourceCodeAccessor.cpp
--- a/Engine/Plugins/Developer/VisualStudioSourceCodeAccess/Source/VisualStudioSourceCodeAccess/Private/VisualStudioSourceCodeAccessor.cpp
+++ b/Engine/Plugins/Developer/VisualStudioSourceCodeAccess/Source/VisualStudioSourceCodeAccess/Private/VisualStudioSourceCodeAccessor.cpp
@@ -89,7 +89,9 @@ static void OnModuleCompileStarted(bool bIsAsyncCompile)
 
 int32 GetVisualStudioVersionForCompiler()
 {
-#if _MSC_VER >= 1920
+#if _MSC_VER >= 1930
+	return 17; // Visual Studio 2022
+#elif _MSC_VER >= 1920
 	return 16; // Visual Studio 2019
 #elif _MSC_VER >= 1910
 	return 15; // Visual Studio 2017
@@ -158,6 +160,7 @@ void FVisualStudioSourceCodeAccessor::RefreshAvailability()
 {
 	Locations.Reset();
 
+	AddVisualStudioVersionUsingVisualStudioSetupAPI(17); // Visual Studio 2022
 	AddVisualStudioVersionUsingVisualStudioSetupAPI(16); // Visual Studio 2019
 	AddVisualStudioVersionUsingVisualStudioSetupAPI(15); // Visual Studio 2017
 	AddVisualStudioVersion(14); // Visual Studio 2015
diff --git a/Engine/Source/Developer/SourceCodeAccess/Private/SourceCodeAccessModule.cpp b/Engine/Source/Developer/SourceCodeAccess/Private/SourceCodeAccessModule.cpp
--- a/Engine/Source/Developer/SourceCodeAccess/Private/SourceCodeAccessModule.cpp
+++ b/Engine/Source/Developer/SourceCodeAccess/Private/SourceCodeAccessModule.cpp
@@ -73,7 +73,7 @@ bool FSourceCodeAccessModule::CanCompileSourceCode() const
 {
 #if PLATFORM_WINDOWS
 	// Need to have Visual Studio installed to compile on Windows, regardless of chosen IDE
-	return IsSourceCodeAccessorAvailable("VisualStudio2017") || IsSourceCodeAccessorAvailable("VisualStudio2019");
+	return IsSourceCodeAccessorAvailable("VisualStudio2017") || IsSourceCodeAccessorAvailable("VisualStudio2019") || IsSourceCodeAccessorAvailable("VisualStudio2022");
 #else
 	// Default behavior
 	return CanAccessSourceCode();
diff --git a/Engine/Source/Developer/Windows/WindowsTargetPlatform/Classes/WindowsTargetSettings.h b/Engine/Source/Developer/Windows/WindowsTargetPlatform/Classes/WindowsTargetSettings.h
--- a/Engine/Source/Developer/Windows/WindowsTargetPlatform/Classes/WindowsTargetSettings.h
+++ b/Engine/Source/Developer/Windows/WindowsTargetPlatform/Classes/WindowsTargetSettings.h
@@ -18,6 +18,7 @@ enum class ECompilerVersion : uint8
 	VisualStudio2015 = 1 UMETA(DisplayName = "Visual Studio 2015 (deprecated)"),
 	VisualStudio2017 = 2 UMETA(DisplayName = "Visual Studio 2017"),
 	VisualStudio2019 = 3 UMETA(DisplayName = "Visual Studio 2019"),
+	VisualStudio2022 = 4 UMETA(DisplayName = "Visual Studio 2022"),
 };
 UENUM()
 enum class EDefaultGraphicsRHI : uint8
diff --git a/Engine/Source/Runtime/Core/Private/ProfilingDebugging/Windows/VSPerfExternalProfiler.cpp b/Engine/Source/Runtime/Core/Private/ProfilingDebugging/Windows/VSPerfExternalProfiler.cpp
--- a/Engine/Source/Runtime/Core/Private/ProfilingDebugging/Windows/VSPerfExternalProfiler.cpp
+++ b/Engine/Source/Runtime/Core/Private/ProfilingDebugging/Windows/VSPerfExternalProfiler.cpp
@@ -104,7 +104,9 @@ public:
 		// Try to load the VSPerf DLL
 		// NOTE: VSPerfXXX.dll is installed into /Windows/System32 when Visual Studio is installed.  The XXX is the version number of
 		// Visual Studio.  For example, for Visual Studio 2013, the file name is VSPerf120.dll.
-#if _MSC_VER >= 1920
+#if _MSC_VER >= 1930
+		DLLHandle = FPlatformProcess::GetDllHandle(TEXT("VSPerf170.dll"));	// Visual Studio 2022
+#elif _MSC_VER >= 1920
 		DLLHandle = FPlatformProcess::GetDllHandle(TEXT("VSPerf160.dll"));	// Visual Studio 2019
 #elif _MSC_VER >= 1910
 		DLLHandle = FPlatformProcess::GetDllHandle(TEXT("VSPerf150.dll"));	// Visual Studio 2017
