Deprecate VS2017 source code accessor

#rnx

diff --git a/Engine/Plugins/Developer/VisualStudioSourceCodeAccess/Source/VisualStudioSourceCodeAccess/Private/VisualStudioSourceCodeAccessModule.cpp b/Engine/Plugins/Developer/VisualStudioSourceCodeAccess/Source/VisualStudioSourceCodeAccess/Private/VisualStudioSourceCodeAccessModule.cpp
--- a/Engine/Plugins/Developer/VisualStudioSourceCodeAccess/Source/VisualStudioSourceCodeAccess/Private/VisualStudioSourceCodeAccessModule.cpp
+++ b/Engine/Plugins/Developer/VisualStudioSourceCodeAccess/Source/VisualStudioSourceCodeAccess/Private/VisualStudioSourceCodeAccessModule.cpp
@@ -21,8 +21,6 @@ void FVisualStudioSourceCodeAccessModule::StartupModule()
 	// from the solution on disk.
 	RegisterWrapper("VisualStudio2022", LOCTEXT("VisualStudio2022", "Visual Studio 2022"), LOCTEXT("UsingVisualStudio2022", "Open source code files in Visual Studio 2022"));
 	RegisterWrapper("VisualStudio2019", LOCTEXT("VisualStudio2019", "Visual Studio 2019"), LOCTEXT("UsingVisualStudio2019", "Open source code files in Visual Studio 2019"));
-	RegisterWrapper("VisualStudio2017", LOCTEXT("VisualStudio2017", "Visual Studio 2017"), LOCTEXT("UsingVisualStudio2017", "Open source code files in Visual Studio 2017"));
-	RegisterWrapper("VisualStudio2015", LOCTEXT("VisualStudio2015", "Visual Studio 2015"), LOCTEXT("UsingVisualStudio2015", "Open source code files in Visual Studio 2015"));
 
 	// Bind our source control provider to the editor
 	IModularFeatures::Get().RegisterModularFeature(TEXT("SourceCodeAccessor"), &VisualStudioSourceCodeAccessor.Get() );
diff --git a/Engine/Source/Developer/SourceCodeAccess/Private/SourceCodeAccessModule.cpp b/Engine/Source/Developer/SourceCodeAccess/Private/SourceCodeAccessModule.cpp
--- a/Engine/Source/Developer/SourceCodeAccess/Private/SourceCodeAccessModule.cpp
+++ b/Engine/Source/Developer/SourceCodeAccess/Private/SourceCodeAccessModule.cpp
@@ -73,7 +73,7 @@ bool FSourceCodeAccessModule::CanCompileSourceCode() const
 {
 #if PLATFORM_WINDOWS
 	// Need to have Visual Studio installed to compile on Windows, regardless of chosen IDE
-	return IsSourceCodeAccessorAvailable("VisualStudio2017") || IsSourceCodeAccessorAvailable("VisualStudio2019") || IsSourceCodeAccessorAvailable("VisualStudio2022");
+	return IsSourceCodeAccessorAvailable("VisualStudio2019") || IsSourceCodeAccessorAvailable("VisualStudio2022");
 #else
 	// Default behavior
 	return CanAccessSourceCode();
diff --git a/Engine/Source/Developer/Windows/WindowsTargetPlatform/Classes/WindowsTargetSettings.h b/Engine/Source/Developer/Windows/WindowsTargetPlatform/Classes/WindowsTargetSettings.h
--- a/Engine/Source/Developer/Windows/WindowsTargetPlatform/Classes/WindowsTargetSettings.h
+++ b/Engine/Source/Developer/Windows/WindowsTargetPlatform/Classes/WindowsTargetSettings.h
@@ -16,7 +16,7 @@ enum class ECompilerVersion : uint8
 {
 	Default = 0,
 	VisualStudio2015 = 1 UMETA(DisplayName = "Visual Studio 2015 (deprecated)"),
-	VisualStudio2017 = 2 UMETA(DisplayName = "Visual Studio 2017"),
+	VisualStudio2017 = 2 UMETA(DisplayName = "Visual Studio 2017 (deprecated)"),
 	VisualStudio2019 = 3 UMETA(DisplayName = "Visual Studio 2019"),
 	VisualStudio2022 = 4 UMETA(DisplayName = "Visual Studio 2022"),
 };
