UnrealBuildTool: Fix failure to open/browse to source file in Visual Studio
* Always update PrimaryProjectName.txt when a project specific solution is generated
* -Game solution is always placed under the project directory, update SourceCodeAccessor according, this code was out of date

#rnx

diff --git a/Engine/Plugins/Developer/RiderSourceCodeAccess/Source/RiderSourceCodeAccess/Private/RiderSourceCodeAccessor.cpp b/Engine/Plugins/Developer/RiderSourceCodeAccess/Source/RiderSourceCodeAccess/Private/RiderSourceCodeAccessor.cpp
--- a/Engine/Plugins/Developer/RiderSourceCodeAccess/Source/RiderSourceCodeAccess/Private/RiderSourceCodeAccessor.cpp
+++ b/Engine/Plugins/Developer/RiderSourceCodeAccess/Source/RiderSourceCodeAccess/Private/RiderSourceCodeAccessor.cpp
@@ -325,9 +325,9 @@ void FRiderSourceCodeAccessor::CachePathToSln() const
 	}
 	else
 	{
-		CachedSolutionPath = FPaths::ConvertRelativePathToFull(FPaths::ProjectDir());
+		const FProjectDescriptor* CurrentProject = IProjectManager::Get().GetCurrentProject();
 
-		if (!FUProjectDictionary(FPaths::RootDir()).IsForeignProject(CachedSolutionPath))
+		if (CurrentProject == nullptr || CurrentProject->Modules.Num() == 0)
 		{
 			FString PrimaryProjectName;
 			if (!FFileHelper::LoadFileToString(PrimaryProjectName, *(FPaths::EngineIntermediateDir() / TEXT("ProjectFiles/PrimaryProjectName.txt"))))
@@ -338,17 +338,13 @@ void FRiderSourceCodeAccessor::CachePathToSln() const
 		}
 		else
 		{
-			const FProjectDescriptor* CurrentProject = IProjectManager::Get().GetCurrentProject();
-
-			if (CurrentProject == nullptr || CurrentProject->Modules.Num() == 0)
-			{
-				CachedSolutionPath = TEXT("");
-			}
-			else
+			CachedSolutionPath = FPaths::ConvertRelativePathToFull(FPaths::ProjectDir());
+			FString BaseName;
+			if (!FFileHelper::LoadFileToString(BaseName, *(FPaths::EngineIntermediateDir() / TEXT("ProjectFiles/PrimaryProjectName.txt"))))
 			{
-				const FString BaseName = FApp::HasProjectName() ? FApp::GetProjectName() : FPaths::GetBaseFilename(CachedSolutionPath);
-				CachedSolutionPath = FPaths::Combine(CachedSolutionPath, BaseName + TEXT(".sln"));
+				BaseName = FApp::HasProjectName() ? FApp::GetProjectName() : FPaths::GetBaseFilename(CachedSolutionPath);
 			}
+			CachedSolutionPath = FPaths::Combine(CachedSolutionPath, BaseName + TEXT(".sln"));
 		}
 	}
 }
diff --git a/Engine/Plugins/Developer/VisualStudioSourceCodeAccess/Source/VisualStudioSourceCodeAccess/Private/VisualStudioSourceCodeAccessor.cpp b/Engine/Plugins/Developer/VisualStudioSourceCodeAccess/Source/VisualStudioSourceCodeAccess/Private/VisualStudioSourceCodeAccessor.cpp
--- a/Engine/Plugins/Developer/VisualStudioSourceCodeAccess/Source/VisualStudioSourceCodeAccess/Private/VisualStudioSourceCodeAccessor.cpp
+++ b/Engine/Plugins/Developer/VisualStudioSourceCodeAccess/Source/VisualStudioSourceCodeAccess/Private/VisualStudioSourceCodeAccessor.cpp
@@ -1406,9 +1406,9 @@ FString FVisualStudioSourceCodeAccessor::GetSolutionPath() const
 		}
 		else
 		{
-			CachedSolutionPath = FPaths::ConvertRelativePathToFull(FPaths::ProjectDir());
+			const FProjectDescriptor* CurrentProject = IProjectManager::Get().GetCurrentProject();
 
-			if (!FUProjectDictionary::GetDefault().IsForeignProject(CachedSolutionPath))
+			if (CurrentProject == nullptr || CurrentProject->Modules.Num() == 0)
 			{
 				FString PrimaryProjectName;
 				if (!FFileHelper::LoadFileToString(PrimaryProjectName, *(FPaths::EngineIntermediateDir() / TEXT("ProjectFiles/PrimaryProjectName.txt"))))
@@ -1419,17 +1419,12 @@ FString FVisualStudioSourceCodeAccessor::GetSolutionPath() const
 			}
 			else
 			{
-				const FProjectDescriptor* CurrentProject = IProjectManager::Get().GetCurrentProject();
-
-				if (CurrentProject == nullptr || CurrentProject->Modules.Num() == 0)
-				{
-					CachedSolutionPath = TEXT("");
-				}
-				else
+				FString BaseName;
+				if (!FFileHelper::LoadFileToString(BaseName, *(FPaths::EngineIntermediateDir() / TEXT("ProjectFiles/PrimaryProjectName.txt"))))
 				{
-					FString BaseName = FApp::HasProjectName() ? FApp::GetProjectName() : FPaths::GetBaseFilename(CachedSolutionPath);
-					CachedSolutionPath = FPaths::Combine(CachedSolutionPath, BaseName + TEXT(".sln"));
+					BaseName = FApp::HasProjectName() ? FApp::GetProjectName() : FPaths::GetBaseFilename(CachedSolutionPath);
 				}
+				CachedSolutionPath = FPaths::Combine(CachedSolutionPath, BaseName + TEXT(".sln"));
 			}
 		}
 	}
diff --git a/Engine/Source/Programs/UnrealBuildTool/ProjectFiles/ProjectFileGenerator.cs b/Engine/Source/Programs/UnrealBuildTool/ProjectFiles/ProjectFileGenerator.cs
--- a/Engine/Source/Programs/UnrealBuildTool/ProjectFiles/ProjectFileGenerator.cs
+++ b/Engine/Source/Programs/UnrealBuildTool/ProjectFiles/ProjectFileGenerator.cs
@@ -866,6 +866,9 @@ namespace UnrealBuildTool
 				// Set the project file name
 				PrimaryProjectName = OnlyGameProject.GetFileNameWithoutExtension();
 
+				// Write out the name of the primary project file, so the runtime knows to use it
+				WritePrimaryProjectNameTxt = true;
+
 				if (!DirectoryReference.Exists(DirectoryReference.Combine(PrimaryProjectPath, "Source")))
 				{
 					if (!DirectoryReference.Exists(DirectoryReference.Combine(PrimaryProjectPath, "Intermediate", "Source")))
