PR #10127: Add missing SupportedTargetPlatforms for VisualStudioSourceCodeAccess (Contributed by [REDACTED_PERSON])

https://jira.it.epicgames.com/browse/UE-177285

diff --git a/Engine/Plugins/Developer/VisualStudioSourceCodeAccess/VisualStudioSourceCodeAccess.uplugin b/Engine/Plugins/Developer/VisualStudioSourceCodeAccess/VisualStudioSourceCodeAccess.uplugin
--- a/Engine/Plugins/Developer/VisualStudioSourceCodeAccess/VisualStudioSourceCodeAccess.uplugin
+++ b/Engine/Plugins/Developer/VisualStudioSourceCodeAccess/VisualStudioSourceCodeAccess.uplugin
@@ -15,6 +15,7 @@
 	"IsBetaVersion" : false,
 	"Installed" : false,
 	"SupportedPrograms": [ "UnrealFrontend", "UnrealInsights" ],
+	"SupportedTargetPlatforms":  [ "Win64" ],
 	"Modules" :
 	[
 		{
