Undo //UE5/Main/Engine/Plugins/Developer/VisualStudioSourceCodeAccess/... changelist 25336010

diff --git a/Engine/Plugins/Developer/VisualStudioSourceCodeAccess/VisualStudioSourceCodeAccess.uplugin b/Engine/Plugins/Developer/VisualStudioSourceCodeAccess/VisualStudioSourceCodeAccess.uplugin
--- a/Engine/Plugins/Developer/VisualStudioSourceCodeAccess/VisualStudioSourceCodeAccess.uplugin
+++ b/Engine/Plugins/Developer/VisualStudioSourceCodeAccess/VisualStudioSourceCodeAccess.uplugin
@@ -15,7 +15,6 @@
 	"IsBetaVersion" : false,
 	"Installed" : false,
 	"SupportedPrograms": [ "UnrealFrontend", "UnrealInsights" ],
-	"SupportedTargetPlatforms":  [ "Win64" ],
 	"Modules" :
 	[
 		{
