Try and select the best solution when opening files via VS code accessors

diff --git a/Engine/Plugins/Developer/VisualStudioSourceCodeAccess/Source/VisualStudioSourceCodeAccess/Private/VisualStudioSourceCodeAccessor.cpp b/Engine/Plugins/Developer/VisualStudioSourceCodeAccess/Source/VisualStudioSourceCodeAccess/Private/VisualStudioSourceCodeAccessor.cpp
--- a/Engine/Plugins/Developer/VisualStudioSourceCodeAccess/Source/VisualStudioSourceCodeAccess/Private/VisualStudioSourceCodeAccessor.cpp
+++ b/Engine/Plugins/Developer/VisualStudioSourceCodeAccess/Source/VisualStudioSourceCodeAccess/Private/VisualStudioSourceCodeAccessor.cpp
@@ -18,6 +18,7 @@
 #include "Misc/MonitoredProcess.h"
 #include "Serialization/JsonReader.h"
 #include "Serialization/JsonSerializer.h"
+#include "Algo/Transform.h"
 //#include "GameProjectGenerationModule.h"
 
 #if WITH_EDITOR
@@ -193,6 +194,87 @@ bool IsVisualStudioDTEMoniker(const FString& InName, const TArray<FVisualStudioS
 	return false;
 }
 
+/** Query all opened instances of Visual Studio to retrieve their respective opened solutions. */
+TArray<FString> RetrieveOpenedVisualStudioSolutionsViaDTE(const TArray<FVisualStudioSourceCodeAccessor::VisualStudioLocation>& InLocations)
+{
+	TArray<FString> OpenedSolutions;
+	// Open the Running Object Table (ROT)
+	IRunningObjectTable* RunningObjectTable;
+	if (SUCCEEDED(GetRunningObjectTable(0, &RunningObjectTable)) && RunningObjectTable)
+	{
+		IEnumMoniker* MonikersTable;
+		if (SUCCEEDED(RunningObjectTable->EnumRunning(&MonikersTable)))
+		{
+			MonikersTable->Reset();
+
+			// Look for all visual studio instances in the ROT
+			IMoniker* CurrentMoniker;
+			while (MonikersTable->Next(1, &CurrentMoniker, NULL) == S_OK)
+			{
+				IBindCtx* BindContext;
+				LPOLESTR OutName;
+				if (SUCCEEDED(CreateBindCtx(0, &BindContext)) && SUCCEEDED(CurrentMoniker->GetDisplayName(BindContext, NULL, &OutName)))
+				{
+					if (IsVisualStudioDTEMoniker(FString(OutName), InLocations))
+					{
+						TComPtr<IUnknown> ComObject;
+						if (SUCCEEDED(RunningObjectTable->GetObject(CurrentMoniker, &ComObject)))
+						{
+							TComPtr<EnvDTE::_DTE> TempDTE;
+							if (SUCCEEDED(TempDTE.FromQueryInterface(__uuidof(EnvDTE::_DTE), ComObject)))
+							{
+								// Get the solution path for this instance
+								TComPtr<EnvDTE::_Solution> Solution;
+								BSTR OutPath = nullptr;
+								if (SUCCEEDED(TempDTE->get_Solution(&Solution)) &&
+									SUCCEEDED(Solution->get_FullName(&OutPath)))
+								{
+									FString Filename(OutPath);
+									FPaths::NormalizeFilename(Filename);
+
+									OpenedSolutions.AddUnique(Filename);
+
+									SysFreeString(OutPath);
+								}
+								else
+								{
+									UE_LOG(LogVSAccessor, Warning, TEXT("Visual Studio is open but could not be queried - it may be blocked by a modal operation"));
+								}
+							}
+							else
+							{
+								UE_LOG(LogVSAccessor, Warning, TEXT("Could not get DTE interface from returned Visual Studio instance"));
+							}
+						}
+						else
+						{
+							UE_LOG(LogVSAccessor, Warning, TEXT("Couldn't get Visual Studio COM object"));
+						}
+					}
+				}
+				else
+				{
+					UE_LOG(LogVSAccessor, Warning, TEXT("Couldn't get display name"));
+				}
+				BindContext->Release();
+				CurrentMoniker->Release();
+			}
+			MonikersTable->Release();
+		}
+		else
+		{
+			UE_LOG(LogVSAccessor, Warning, TEXT("Couldn't enumerate ROT table"));
+		}
+		RunningObjectTable->Release();
+	}
+	else
+	{
+		UE_LOG(LogVSAccessor, Warning, TEXT("Couldn't get ROT table"));
+	}
+
+	return OpenedSolutions;
+}
+
 /** Accesses the correct visual studio instance if possible. */
 EAccessVisualStudioResult AccessVisualStudioViaDTE(TComPtr<EnvDTE::_DTE>& OutDTE, const FString& InSolutionPath, const TArray<FVisualStudioSourceCodeAccessor::VisualStudioLocation>& InLocations)
 {
@@ -322,7 +404,7 @@ bool FVisualStudioSourceCodeAccessor::OpenVisualStudioSolutionViaDTE()
 			// Automatically fail if there's already an attempt in progress
 			if (!IsVSLaunchInProgress())
 			{
-				bSuccess = RunVisualStudioAndOpenSolution();
+				bSuccess = RunVisualStudioAndOpenSolution(SolutionPath);
 			}
 		}
 		break;
@@ -339,6 +421,84 @@ bool FVisualStudioSourceCodeAccessor::OpenVisualStudioSolutionViaDTE()
 	return bSuccess;
 }
 
+FString FVisualStudioSourceCodeAccessor::RetrieveSolutionForFileOpenRequests(const TArray<FileOpenRequest>& Requests, const TArray<FString>& CurrentlyOpenedSolutions) const
+{
+	// Based on the files being requested, make an educated guess as to which is the most appropriate solution to open them all by finding the corresponding .sln/.slnf files in the folder hierarchy: 
+	struct FSolutionInfo
+	{
+		FSolutionInfo(const FString& InSolutionFile)
+			: SolutionFile(InSolutionFile)
+		{}
+
+		// Describes the state of a solution file wrt the currently opened solutions (ordered by priority)
+		enum class EOpenedSolutionState : uint8
+		{
+			CurrentlyOpenedExactMatch = 0, // The solution is currently opened in visual studio
+			CurrentlyOpened, // A solution with same file name but not the same absolute path is currently opened in visual studio
+			NotOpened, // The solution is not currently opened in Visual Studio
+		};
+		EOpenedSolutionState OpenedSolutionState = FSolutionInfo::EOpenedSolutionState::NotOpened;
+		int32 RefCount = 0;
+		FString SolutionFile;
+	};
+
+	TArray<FSolutionInfo> SolutionFileInfos;
+	for (const FileOpenRequest& Request : Requests)
+	{
+		FString CurrentPath = FPaths::GetPath(Request.FullPath);
+		while (!CurrentPath.IsEmpty())
+		{
+			TArray<FString> FilesInDirectory;
+			IFileManager::Get().FindFiles(FilesInDirectory, *CurrentPath, TEXT(".sln"));
+			IFileManager::Get().FindFiles(FilesInDirectory, *CurrentPath, TEXT(".slnf"));
+			for (const FString& FileInDirectory : FilesInDirectory)
+			{
+				FString AbsoluteFileName = CurrentPath / FileInDirectory;
+				FPaths::NormalizeFilename(AbsoluteFileName);
+
+				FSolutionInfo* SolutionInfo = SolutionFileInfos.FindByPredicate([&AbsoluteFileName](const FSolutionInfo& InSolutionInfo) { return InSolutionInfo.SolutionFile == AbsoluteFileName; });
+				if (SolutionInfo == nullptr)
+				{
+					SolutionInfo = &SolutionFileInfos.Add_GetRef(FSolutionInfo(AbsoluteFileName));
+
+					for (const FString& OpenedSolution : CurrentlyOpenedSolutions)
+					{
+						if (OpenedSolution.Equals(AbsoluteFileName, ESearchCase::IgnoreCase))
+						{
+							SolutionInfo->OpenedSolutionState = FSolutionInfo::EOpenedSolutionState::CurrentlyOpenedExactMatch;
+							break;
+						}
+						if (FPaths::GetCleanFilename(OpenedSolution).Equals(FileInDirectory, ESearchCase::IgnoreCase))
+						{
+							SolutionInfo->OpenedSolutionState = FSolutionInfo::EOpenedSolutionState::CurrentlyOpened;
+						}
+					}
+				}
+				++SolutionInfo->RefCount;
+			}
+			CurrentPath = FPaths::GetPath(CurrentPath);
+		}
+	}
+
+	// Now that we have a list of all solutions that could be used to open all these files, pick the best one : 
+	SolutionFileInfos.Sort([](const FSolutionInfo& InLHS, const FSolutionInfo& InRHS)
+	{
+		// Sort by ref count first, so that the most requested solution comes out on top : 
+		if (InLHS.RefCount < InRHS.RefCount)
+		{
+			return false;
+		}
+		else if (InLHS.RefCount == InRHS.RefCount)
+		{
+			return (static_cast<uint8>(InLHS.OpenedSolutionState) < static_cast<uint8>(InRHS.OpenedSolutionState));
+		}
+
+		return true;
+	});
+
+	return SolutionFileInfos.IsEmpty() ? FString() : SolutionFileInfos[0].SolutionFile;
+}
+
 bool FVisualStudioSourceCodeAccessor::OpenVisualStudioFilesInternalViaDTE(const TArray<FileOpenRequest>& Requests, bool& bWasDeferred)
 {
 	ISourceCodeAccessModule& SourceCodeAccessModule = FModuleManager::LoadModuleChecked<ISourceCodeAccessModule>(TEXT("SourceCodeAccess"));
@@ -350,10 +510,20 @@ bool FVisualStudioSourceCodeAccessor::OpenVisualStudioFilesInternalViaDTE(const
 		return false;
 	}
 	
+	FString SolutionPath = GetSolutionPath();
+	TArray<VisualStudioLocation> InstalledLocations = GetPrioritizedVisualStudioVersions(SolutionPath);
+
+	// If no solution is specified, make an educated guess based on the files being requested by finding the corresponding .sln/.slnf files: 
+	if (SolutionPath.IsEmpty())
+	{
+		TArray<FString> CurrentlyOpenedSolutions = RetrieveOpenedVisualStudioSolutionsViaDTE(InstalledLocations);
+		SolutionPath = RetrieveSolutionForFileOpenRequests(Requests, CurrentlyOpenedSolutions);
+	}
+
 	bool bDefer = false, bSuccess = false;
 	TComPtr<EnvDTE::_DTE> DTE;
-	const FString SolutionPath = GetSolutionPath();
-	switch (AccessVisualStudioViaDTE(DTE, SolutionPath, GetPrioritizedVisualStudioVersions(SolutionPath)))
+	
+	switch (AccessVisualStudioViaDTE(DTE, SolutionPath, InstalledLocations))
 	{
 	case EAccessVisualStudioResult::VSInstanceIsOpen:
 		{
@@ -455,7 +625,7 @@ bool FVisualStudioSourceCodeAccessor::OpenVisualStudioFilesInternalViaDTE(const
 				if ( !IsVSLaunchInProgress() )
 				{
 					// If there's no valid instance of VS running, run one if we have it installed
-					if ( !RunVisualStudioAndOpenSolution() )
+					if ( !RunVisualStudioAndOpenSolution(SolutionPath) )
 					{
 						bDefer = false;
 					}
@@ -693,6 +863,91 @@ bool GetProcessCommandLine(const ::DWORD InProcessID, FString& OutCommandLine)
 	return Data.OutHwnd;
 }
 
+/** Query all opened instances of Visual Studio to retrieve their respective opened solutions. */
+TArray<FString> RetrieveOpenedVisualStudioSolutionsViaProcess(const TArray<FVisualStudioSourceCodeAccessor::VisualStudioLocation>& InLocations)
+{
+	TArray<FString> OpenedSolutions;
+
+	::HANDLE hProcessSnap = ::CreateToolhelp32Snapshot(TH32CS_SNAPPROCESS, 0);
+	if (hProcessSnap != INVALID_HANDLE_VALUE)
+	{
+		MODULEENTRY32 ModuleEntry;
+		ModuleEntry.dwSize = sizeof(MODULEENTRY32);
+
+		PROCESSENTRY32 ProcEntry;
+		ProcEntry.dwSize = sizeof(PROCESSENTRY32);
+
+		// We enumerate the locations as the outer loop to ensure we find our preferred process type first
+		// If we did this as the inner loop, then we'd get the first process that matched any location, even if it wasn't our preference
+		for (const auto& Location : InLocations)
+		{
+			for (::BOOL bHasProcess = ::Process32First(hProcessSnap, &ProcEntry); bHasProcess; bHasProcess = ::Process32Next(hProcessSnap, &ProcEntry))
+			{
+				::HANDLE hModuleSnap = ::CreateToolhelp32Snapshot(TH32CS_SNAPMODULE, ProcEntry.th32ProcessID);
+				if (hModuleSnap != INVALID_HANDLE_VALUE)
+				{
+					if (::Module32First(hModuleSnap, &ModuleEntry))
+					{
+						FString ProcPath = ModuleEntry.szExePath;
+
+						if (ProcPath == Location.ExecutablePath)
+						{
+							// Without DTE we can't accurately verify that the Visual Studio instance has the correct solution open,
+							// however, if we've opened it (or it's opened the solution directly), then the solution path will
+							// exist somewhere in the command line for the process
+							FString CommandLine;
+							if (GetProcessCommandLine(ProcEntry.th32ProcessID, CommandLine))
+							{
+								TArray<FString> Tokens;
+								CommandLine.ParseIntoArray(Tokens, TEXT(" "), /*InCullEmpty = */true);
+								FString SolutionName;
+								TArray<FString> SolutionNames;
+								Algo::TransformIf(Tokens, SolutionNames, 
+									[](const FString& InToken) { return InToken.Contains(TEXT(".slnf"), ESearchCase::IgnoreCase) || InToken.Contains(TEXT(".slnf"), ESearchCase::IgnoreCase); },
+									[](const FString& InToken) 
+									{ 
+										FString Filename = InToken.TrimStartAndEnd();
+										Filename.RemoveFromStart(TEXT("\""));
+										Filename.RemoveFromEnd(TEXT("\""));
+										FPaths::NormalizeFilename(Filename);
+										return Filename;
+									});
+								
+								if (!SolutionNames.IsEmpty())
+								{
+									OpenedSolutions.AddUnique(SolutionNames[0]);
+								}
+							}
+							else
+							{
+								UE_LOG(LogVSAccessor, Warning, TEXT("Couldn't access module information"));
+							}
+						}
+					}
+					else
+					{
+						UE_LOG(LogVSAccessor, Warning, TEXT("Couldn't access module table"));
+					}
+
+					::CloseHandle(hModuleSnap);
+				}
+				else
+				{
+					UE_LOG(LogVSAccessor, Warning, TEXT("Couldn't access module table"));
+				}
+			}
+		}
+
+		::CloseHandle(hProcessSnap);
+	}
+	else
+	{
+		UE_LOG(LogVSAccessor, Warning, TEXT("Couldn't access process table"));
+	}
+
+	return OpenedSolutions;
+}
+
 EAccessVisualStudioResult AccessVisualStudioViaProcess(::DWORD& OutProcessID, FString& OutExecutablePath, const FString& InSolutionPath, const TArray<FVisualStudioSourceCodeAccessor::VisualStudioLocation>& InLocations)
 {
 	OutProcessID = 0;
@@ -851,7 +1106,7 @@ bool FVisualStudioSourceCodeAccessor::OpenVisualStudioSolutionViaProcess()
 		break;
 
 	case EAccessVisualStudioResult::VSInstanceIsNotOpen:
-		return RunVisualStudioAndOpenSolution();
+		return RunVisualStudioAndOpenSolution(SolutionPath);
 
 	default:
 		// Do nothing if we failed the VS detection, otherwise we could get stuck in a loop of constantly 
@@ -866,8 +1121,18 @@ bool FVisualStudioSourceCodeAccessor::OpenVisualStudioFilesInternalViaProcess(co
 {
 	::DWORD ProcessID = 0;
 	FString Path;
-	const FString SolutionPath = GetSolutionPath();
-	switch (AccessVisualStudioViaProcess(ProcessID, Path, SolutionPath, GetPrioritizedVisualStudioVersions(SolutionPath)))
+	FString SolutionPath = GetSolutionPath();
+
+	TArray<VisualStudioLocation> InstalledLocations = GetPrioritizedVisualStudioVersions(SolutionPath);
+
+	// If no solution is specified, make an educated guess based on the files being requested by finding the corresponding .sln files: 
+	if (SolutionPath.IsEmpty())
+	{
+		TArray<FString> CurrentlyOpenedSolutions = RetrieveOpenedVisualStudioSolutionsViaProcess(InstalledLocations);
+		SolutionPath = RetrieveSolutionForFileOpenRequests(Requests, CurrentlyOpenedSolutions);
+	}
+
+	switch (AccessVisualStudioViaProcess(ProcessID, Path, SolutionPath, InstalledLocations))
 	{
 	case EAccessVisualStudioResult::VSInstanceIsOpen:
 		return RunVisualStudioAndOpenSolutionAndFiles(Path, "", &Requests);
@@ -900,13 +1165,12 @@ bool FVisualStudioSourceCodeAccessor::CanRunVisualStudio(FString& OutPath, const
 	return false;
 }
 
-bool FVisualStudioSourceCodeAccessor::RunVisualStudioAndOpenSolution() const
+bool FVisualStudioSourceCodeAccessor::RunVisualStudioAndOpenSolution(const FString& InSolution) const
 {
 	FString Path;
-	const FString SolutionPath = GetSolutionPath();
-	if (CanRunVisualStudio(Path, SolutionPath))
+	if (CanRunVisualStudio(Path, InSolution))
 	{
-		return RunVisualStudioAndOpenSolutionAndFiles(Path, SolutionPath, nullptr);
+		return RunVisualStudioAndOpenSolutionAndFiles(Path, InSolution, nullptr);
 	}
 
 	return false;
diff --git a/Engine/Plugins/Developer/VisualStudioSourceCodeAccess/Source/VisualStudioSourceCodeAccess/Private/VisualStudioSourceCodeAccessor.h b/Engine/Plugins/Developer/VisualStudioSourceCodeAccess/Source/VisualStudioSourceCodeAccess/Private/VisualStudioSourceCodeAccessor.h
--- a/Engine/Plugins/Developer/VisualStudioSourceCodeAccess/Source/VisualStudioSourceCodeAccess/Private/VisualStudioSourceCodeAccessor.h
+++ b/Engine/Plugins/Developer/VisualStudioSourceCodeAccess/Source/VisualStudioSourceCodeAccess/Private/VisualStudioSourceCodeAccessor.h
@@ -71,7 +71,7 @@ private:
 	bool CanRunVisualStudio(FString& OutPath, const FString& InSolution) const;
 
 	/** Run an instance of visual studio instance if possible. */
-	bool RunVisualStudioAndOpenSolution() const;
+	bool RunVisualStudioAndOpenSolution(const FString& InSolution) const;
 
 	/** Opens a file in the correct running instance of Visual Studio at a line and optionally to a column. */
 	bool OpenVisualStudioFileAtLineInternal(const FString& FullPath, int32 LineNumber, int32 ColumnNumber = 0);
@@ -127,6 +127,15 @@ private:
 	 */
 	bool RunVisualStudioAndOpenSolutionAndFiles(const FString& ExecutablePath, const FString& SolutionPath, const TArray<FileOpenRequest>* const Requests) const;
 
+	/**
+	 * Based on the list of file open requests and the list of solutions currently opened, make an educated guess as to which solution might be the most relevant to open them all
+	 *
+	 * @param	Requests					Array of files to open
+	 * @param	CurrentlyOpenedSolutions	List of solutions currently opened by running VS instances
+	 * @return the absolute file path to the most likely solution or an empty string if no solution was found 
+	 */
+	FString RetrieveSolutionForFileOpenRequests(const TArray<FileOpenRequest>& Requests, const TArray<FString>& CurrentlyOpenedSolutions) const;
+
 #if WITH_VISUALSTUDIO_DTE
 	/** DTE specific implementations */
 	bool OpenVisualStudioSolutionViaDTE();
