Fixed static analysis

#rnx

diff --git a/Engine/Plugins/Developer/VisualStudioSourceCodeAccess/Source/VisualStudioSourceCodeAccess/Private/VisualStudioSourceCodeAccessor.cpp b/Engine/Plugins/Developer/VisualStudioSourceCodeAccess/Source/VisualStudioSourceCodeAccess/Private/VisualStudioSourceCodeAccessor.cpp
--- a/Engine/Plugins/Developer/VisualStudioSourceCodeAccess/Source/VisualStudioSourceCodeAccess/Private/VisualStudioSourceCodeAccessor.cpp
+++ b/Engine/Plugins/Developer/VisualStudioSourceCodeAccess/Source/VisualStudioSourceCodeAccess/Private/VisualStudioSourceCodeAccessor.cpp
@@ -903,7 +903,7 @@ TArray<FString> RetrieveOpenedVisualStudioSolutionsViaProcess(const TArray<FVisu
 								FString SolutionName;
 								TArray<FString> SolutionNames;
 								Algo::TransformIf(Tokens, SolutionNames, 
-									[](const FString& InToken) { return InToken.Contains(TEXT(".slnf"), ESearchCase::IgnoreCase) || InToken.Contains(TEXT(".slnf"), ESearchCase::IgnoreCase); },
+									[](const FString& InToken) { return InToken.EndsWith(TEXT(".sln"), ESearchCase::IgnoreCase) || InToken.EndsWith(TEXT(".slnf"), ESearchCase::IgnoreCase); },
 									[](const FString& InToken) 
 									{ 
 										FString Filename = InToken.TrimStartAndEnd();
