[Backout] - CL25209946
Original CL Desc
-----------------------------------------------------------------
UnrealBuildTool: Fix failure to open/browse to source file in Visual Studio
* Always update PrimaryProjectName.txt when a project specific solution is generated
* -Game solution is always placed under the project directory, update SourceCodeAccessor according, this code was out of date

#rnx

diff --git a/Engine/Plugins/Developer/RiderSourceCodeAccess/Source/RiderSourceCodeAccess/Private/RiderSourceCodeAccessor.cpp b/Engine/Plugins/Developer/RiderSourceCodeAccess/Source/RiderSourceCodeAccess/Private/RiderSourceCodeAccessor.cpp
--- a/Engine/Plugins/Developer/RiderSourceCodeAccess/Source/RiderSourceCodeAccess/Private/RiderSourceCodeAccessor.cpp
+++ b/Engine/Plugins/Developer/RiderSourceCodeAccess/Source/RiderSourceCodeAccess/Private/RiderSourceCodeAccessor.cpp
@@ -325,9 +325,9 @@ void FRiderSourceCodeAccessor::CachePathToSln() const
 	}
 	else
 	{
-		const FProjectDescriptor* CurrentProject = IProjectManager::Get().GetCurrentProject();
+		CachedSolutionPath = FPaths::ConvertRelativePathToFull(FPaths::ProjectDir());
 
-		if (CurrentProject == nullptr || CurrentProject->Modules.Num() == 0)
+		if (!FUProjectDictionary(FPaths::RootDir()).IsForeignProject(CachedSolutionPath))
 		{
 			FString PrimaryProjectName;
 			if (!FFileHelper::LoadFileToString(PrimaryProjectName, *(FPaths::EngineIntermediateDir() / TEXT("ProjectFiles/PrimaryProjectName.txt"))))
@@ -338,13 +338,17 @@ void FRiderSourceCodeAccessor::CachePathToSln() const
 		}
 		else
 		{
-			CachedSolutionPath = FPaths::ConvertRelativePathToFull(FPaths::ProjectDir());
-			FString BaseName;
-			if (!FFileHelper::LoadFileToString(BaseName, *(FPaths::EngineIntermediateDir() / TEXT("ProjectFiles/PrimaryProjectName.txt"))))
+			const FProjectDescriptor* CurrentProject = IProjectManager::Get().GetCurrentProject();
+
+			if (CurrentProject == nullptr || CurrentProject->Modules.Num() == 0)
+			{
+				CachedSolutionPath = TEXT("");
+			}
+			else
 			{
-				BaseName = FApp::HasProjectName() ? FApp::GetProjectName() : FPaths::GetBaseFilename(CachedSolutionPath);
+				const FString BaseName = FApp::HasProjectName() ? FApp::GetProjectName() : FPaths::GetBaseFilename(CachedSolutionPath);
+				CachedSolutionPath = FPaths::Combine(CachedSolutionPath, BaseName + TEXT(".sln"));
 			}
-			CachedSolutionPath = FPaths::Combine(CachedSolutionPath, BaseName + TEXT(".sln"));
 		}
 	}
 }
diff --git a/Engine/Plugins/Developer/VisualStudioSourceCodeAccess/Source/VisualStudioSourceCodeAccess/Private/VisualStudioSourceCodeAccessor.cpp b/Engine/Plugins/Developer/VisualStudioSourceCodeAccess/Source/VisualStudioSourceCodeAccess/Private/VisualStudioSourceCodeAccessor.cpp
--- a/Engine/Plugins/Developer/VisualStudioSourceCodeAccess/Source/VisualStudioSourceCodeAccess/Private/VisualStudioSourceCodeAccessor.cpp
+++ b/Engine/Plugins/Developer/VisualStudioSourceCodeAccess/Source/VisualStudioSourceCodeAccess/Private/VisualStudioSourceCodeAccessor.cpp
@@ -1670,9 +1670,9 @@ FString FVisualStudioSourceCodeAccessor::GetSolutionPath() const
 		}
 		else
 		{
-			const FProjectDescriptor* CurrentProject = IProjectManager::Get().GetCurrentProject();
+			CachedSolutionPath = FPaths::ConvertRelativePathToFull(FPaths::ProjectDir());
 
-			if (CurrentProject == nullptr || CurrentProject->Modules.Num() == 0)
+			if (!FUProjectDictionary::GetDefault().IsForeignProject(CachedSolutionPath))
 			{
 				FString PrimaryProjectName;
 				if (!FFileHelper::LoadFileToString(PrimaryProjectName, *(FPaths::EngineIntermediateDir() / TEXT("ProjectFiles/PrimaryProjectName.txt"))))
@@ -1683,12 +1683,17 @@ FString FVisualStudioSourceCodeAccessor::GetSolutionPath() const
 			}
 			else
 			{
-				FString BaseName;
-				if (!FFileHelper::LoadFileToString(BaseName, *(FPaths::EngineIntermediateDir() / TEXT("ProjectFiles/PrimaryProjectName.txt"))))
+				const FProjectDescriptor* CurrentProject = IProjectManager::Get().GetCurrentProject();
+
+				if (CurrentProject == nullptr || CurrentProject->Modules.Num() == 0)
+				{
+					CachedSolutionPath = TEXT("");
+				}
+				else
 				{
-					BaseName = FApp::HasProjectName() ? FApp::GetProjectName() : FPaths::GetBaseFilename(CachedSolutionPath);
+					FString BaseName = FApp::HasProjectName() ? FApp::GetProjectName() : FPaths::GetBaseFilename(CachedSolutionPath);
+					CachedSolutionPath = FPaths::Combine(CachedSolutionPath, BaseName + TEXT(".sln"));
 				}
-				CachedSolutionPath = FPaths::Combine(CachedSolutionPath, BaseName + TEXT(".sln"));
 			}
 		}
 	}
diff --git a/Engine/Source/Programs/UnrealBuildTool/ProjectFiles/ProjectFileGenerator.cs b/Engine/Source/Programs/UnrealBuildTool/ProjectFiles/ProjectFileGenerator.cs
--- a/Engine/Source/Programs/UnrealBuildTool/ProjectFiles/ProjectFileGenerator.cs
+++ b/Engine/Source/Programs/UnrealBuildTool/ProjectFiles/ProjectFileGenerator.cs
@@ -872,9 +872,6 @@ namespace UnrealBuildTool
 				// Set the project file name
 				PrimaryProjectName = OnlyGameProject.GetFileNameWithoutExtension();
 
-				// Write out the name of the primary project file, so the runtime knows to use it
-				WritePrimaryProjectNameTxt = true;
-
 				if (!bAllowContentOnlyProjects && !DirectoryReference.Exists(DirectoryReference.Combine(PrimaryProjectPath, "Source")))
 				{
 					if (!DirectoryReference.Exists(DirectoryReference.Combine(PrimaryProjectPath, "Intermediate", "Source")))
