When opening code in a solution also check any running instances of Visual Studio to find the best solution to use

#rnx

diff --git a/Engine/Plugins/Developer/VisualStudioSourceCodeAccess/Source/VisualStudioSourceCodeAccess/Private/VisualStudioSourceCodeAccessor.cpp b/Engine/Plugins/Developer/VisualStudioSourceCodeAccess/Source/VisualStudioSourceCodeAccess/Private/VisualStudioSourceCodeAccessor.cpp
--- a/Engine/Plugins/Developer/VisualStudioSourceCodeAccess/Source/VisualStudioSourceCodeAccess/Private/VisualStudioSourceCodeAccessor.cpp
+++ b/Engine/Plugins/Developer/VisualStudioSourceCodeAccess/Source/VisualStudioSourceCodeAccess/Private/VisualStudioSourceCodeAccessor.cpp
@@ -435,12 +435,13 @@ bool FVisualStudioSourceCodeAccessor::OpenVisualStudioFilesInternalViaDTE(const
 	FString SolutionPath = GetSolutionPath();
 	TArray<VisualStudioLocation> InstalledLocations = GetPrioritizedVisualStudioVersions(SolutionPath);
 
-	// If no solution is specified, make an educated guess based on the files being requested by finding the corresponding .sln/.slnf files: 
-	if (SolutionPath.IsEmpty())
+	// Even if a solution is specified, also check open solutions and make an educated guess about the best path based on the files being requested by finding the corresponding .sln/.slnf files: 
+	TArray<FString> CurrentlyOpenedSolutions = RetrieveOpenedVisualStudioSolutionsViaDTE(InstalledLocations);
+	if (!SolutionPath.IsEmpty())
 	{
-		TArray<FString> CurrentlyOpenedSolutions = RetrieveOpenedVisualStudioSolutionsViaDTE(InstalledLocations);
-		SolutionPath = RetrieveSolutionForFileOpenRequests(Requests, CurrentlyOpenedSolutions);
+		CurrentlyOpenedSolutions.Add(SolutionPath);
 	}
+	SolutionPath = RetrieveSolutionForFileOpenRequests(Requests, CurrentlyOpenedSolutions);
 
 	bool bDefer = false, bSuccess = false;
 	TComPtr<EnvDTE::_DTE> DTE;
@@ -1125,12 +1126,13 @@ bool FVisualStudioSourceCodeAccessor::OpenVisualStudioFilesInternalViaProcess(co
 
 	TArray<VisualStudioLocation> InstalledLocations = GetPrioritizedVisualStudioVersions(SolutionPath);
 
-	// If no solution is specified, make an educated guess based on the files being requested by finding the corresponding .sln files: 
-	if (SolutionPath.IsEmpty())
+	// Even if a solution is specified, also check open solutions and make an educated guess about the best path based on the files being requested by finding the corresponding .sln files: 
+	TArray<FString> CurrentlyOpenedSolutions = RetrieveOpenedVisualStudioSolutionsViaProcess(InstalledLocations);
+	if (!SolutionPath.IsEmpty())
 	{
-		TArray<FString> CurrentlyOpenedSolutions = RetrieveOpenedVisualStudioSolutionsViaProcess(InstalledLocations);
-		SolutionPath = RetrieveSolutionForFileOpenRequests(Requests, CurrentlyOpenedSolutions);
+		CurrentlyOpenedSolutions.Add(SolutionPath);
 	}
+	SolutionPath = RetrieveSolutionForFileOpenRequests(Requests, CurrentlyOpenedSolutions);
 
 	switch (AccessVisualStudioViaProcess(ProcessID, Path, SolutionPath, InstalledLocations))
 	{
