UnrealBuildTool: Remove VS2019 as a supported source code access ide

diff --git a/Engine/Plugins/Developer/VisualStudioSourceCodeAccess/Source/VisualStudioSourceCodeAccess/Private/VisualStudioSourceCodeAccessModule.cpp b/Engine/Plugins/Developer/VisualStudioSourceCodeAccess/Source/VisualStudioSourceCodeAccess/Private/VisualStudioSourceCodeAccessModule.cpp
--- a/Engine/Plugins/Developer/VisualStudioSourceCodeAccess/Source/VisualStudioSourceCodeAccess/Private/VisualStudioSourceCodeAccessModule.cpp
+++ b/Engine/Plugins/Developer/VisualStudioSourceCodeAccess/Source/VisualStudioSourceCodeAccess/Private/VisualStudioSourceCodeAccessModule.cpp
@@ -25,7 +25,6 @@ void FVisualStudioSourceCodeAccessModule::StartupModule()
 	// Add all the explicit version wrappers. If one of these is selected, UBT will generate project files in the appropriate format. Editor behavior is still to detect which version to use
 	// from the solution on disk.
 	RegisterWrapper("VisualStudio2022", LOCTEXT("VisualStudio2022", "Visual Studio 2022"), LOCTEXT("UsingVisualStudio2022", "Open source code files in Visual Studio 2022"));
-	RegisterWrapper("VisualStudio2019", LOCTEXT("VisualStudio2019", "Visual Studio 2019"), LOCTEXT("UsingVisualStudio2019", "Open source code files in Visual Studio 2019"));
 
 	// Bind our source control provider to the editor
 	IModularFeatures::Get().RegisterModularFeature(TEXT("SourceCodeAccessor"), &VisualStudioSourceCodeAccessor.Get() );
diff --git a/Engine/Source/Developer/SourceCodeAccess/Private/SourceCodeAccessModule.cpp b/Engine/Source/Developer/SourceCodeAccess/Private/SourceCodeAccessModule.cpp
--- a/Engine/Source/Developer/SourceCodeAccess/Private/SourceCodeAccessModule.cpp
+++ b/Engine/Source/Developer/SourceCodeAccess/Private/SourceCodeAccessModule.cpp
@@ -78,7 +78,7 @@ bool FSourceCodeAccessModule::CanCompileSourceCode() const
 {
 #if PLATFORM_WINDOWS
 	// Need to have Visual Studio installed to compile on Windows, regardless of chosen IDE
-	return IsSourceCodeAccessorAvailable("VisualStudio2019") || IsSourceCodeAccessorAvailable("VisualStudio2022");
+	return IsSourceCodeAccessorAvailable("VisualStudio2022");
 #else
 	// Default behavior
 	return CanAccessSourceCode();
diff --git a/Engine/Source/Developer/Windows/WindowsTargetPlatform/Classes/WindowsTargetSettings.h b/Engine/Source/Developer/Windows/WindowsTargetPlatform/Classes/WindowsTargetSettings.h
--- a/Engine/Source/Developer/Windows/WindowsTargetPlatform/Classes/WindowsTargetSettings.h
+++ b/Engine/Source/Developer/Windows/WindowsTargetPlatform/Classes/WindowsTargetSettings.h
@@ -17,7 +17,7 @@ enum class ECompilerVersion : uint8
 	Default = 0,
 	VisualStudio2015 = 1 UMETA(DisplayName = "Visual Studio 2015 (deprecated)"),
 	VisualStudio2017 = 2 UMETA(DisplayName = "Visual Studio 2017 (deprecated)"),
-	VisualStudio2019 = 3 UMETA(DisplayName = "Visual Studio 2019"),
+	VisualStudio2019 = 3 UMETA(DisplayName = "Visual Studio 2019 (deprecated)"),
 	VisualStudio2022 = 4 UMETA(DisplayName = "Visual Studio 2022"),
 };
 
