PR #11750: Fixes VisualStudioSourceCodeAccessor raw CoCreateInstance call

#rnx

### Original change
- Add a call to ::CoInitialize in FVisualStudioSourceCodeAccessor::AddVisualStudioVersionUsingVisualStudioSetupAPI. The submitted had some use case where this could be called without CoInitialize already being in scope which was causing issues.

### Additional changes
- Added a new struct FCoInitializeScope that calls ::CoInitialize in it's constructor and ::CoUninitialize in it's destructor so that we can keep things in sync and reduce deplicate code such as logging an error if the call fails.
- - Replaced all existing manual calls to ::CoInitialize/::CoUninitialize with this struct.
- ::RefreshAvailability makes serveral calls to AddVisualStudioVersionUsingVisualStudioSetupAPI which in turn will call ::CoInitialize/::CoUninitialize, if we wrap the entire method in a FCoInitializeScope then we will only initialize and destroy once and the calls inside of AddVisualStudioVersionUsingVisualStudioSetupAPI just inc/dec the internal ref count.
-- In practice, since this is all called on the GameThread some other system will have called ::CoInitialize and maintained it's reference count so this makes no difference, but if a stand alone program wanted to use this module it might help.

diff --git a/Engine/Plugins/Developer/VisualStudioSourceCodeAccess/Source/VisualStudioSourceCodeAccess/Private/VisualStudioSourceCodeAccessor.cpp b/Engine/Plugins/Developer/VisualStudioSourceCodeAccess/Source/VisualStudioSourceCodeAccess/Private/VisualStudioSourceCodeAccessor.cpp
--- a/Engine/Plugins/Developer/VisualStudioSourceCodeAccess/Source/VisualStudioSourceCodeAccess/Private/VisualStudioSourceCodeAccessor.cpp
+++ b/Engine/Plugins/Developer/VisualStudioSourceCodeAccess/Source/VisualStudioSourceCodeAccess/Private/VisualStudioSourceCodeAccessor.cpp
@@ -64,6 +64,33 @@ private:
 	BSTR String;
 };
 
+class FCoInitializeScope
+{
+public:
+	FCoInitializeScope()
+	{
+		bInitialized = FWindowsPlatformMisc::CoInitialize();
+		UE_CLOG(!bInitialized, LogVSAccessor, Error, TEXT("ERROR - Could not initialize COM library!"));
+	}
+
+	~FCoInitializeScope()
+	{
+		if (bInitialized)
+		{
+			FWindowsPlatformMisc::CoUninitialize();
+		}
+	}
+
+	bool IsValid() const
+	{
+		return bInitialized;
+	}
+
+private:
+
+	bool bInitialized;
+};
+
 /** The VS query in progress message */
 TWeakPtr<class SNotificationItem> VSNotificationPtr = NULL;
 
@@ -160,9 +187,17 @@ void FVisualStudioSourceCodeAccessor::RefreshAvailability()
 {
 	Locations.Reset();
 
-	AddVisualStudioVersionUsingVisualStudioSetupAPI(17); // Visual Studio 2022
-	AddVisualStudioVersionUsingVisualStudioSetupAPI(16); // Visual Studio 2019
-	AddVisualStudioVersionUsingVisualStudioSetupAPI(15); // Visual Studio 2017
+	// Minor optimization, as each call to AddVisualStudioVersionUsingVisualStudioSetupAPI will make it's
+	// own calls to ::CoInitialize/::CoUninitialize. If we do our own calls here then they will just inc/dec
+	// the internal ref count rather than potentially creating and destroying resources for each call.
+	FCoInitializeScope CoInitialze;
+	if (CoInitialze.IsValid())
+	{
+		AddVisualStudioVersionUsingVisualStudioSetupAPI(17); // Visual Studio 2022
+		AddVisualStudioVersionUsingVisualStudioSetupAPI(16); // Visual Studio 2019
+		AddVisualStudioVersionUsingVisualStudioSetupAPI(15); // Visual Studio 2017
+	}
+
 	AddVisualStudioVersion(14); // Visual Studio 2015
 	AddVisualStudioVersion(12); // Visual Studio 2013
 }
@@ -370,9 +405,9 @@ EAccessVisualStudioResult AccessVisualStudioViaDTE(TComPtr<EnvDTE::_DTE>& OutDTE
 bool FVisualStudioSourceCodeAccessor::OpenVisualStudioSolutionViaDTE()
 {
 	// Initialize the com library, if not already by this thread
-	if (!FWindowsPlatformMisc::CoInitialize())
+	FCoInitializeScope CoInitialze;
+	if (!CoInitialze.IsValid())
 	{
-		UE_LOG(LogVSAccessor, Error, TEXT( "ERROR - Could not initialize COM library!" ));
 		return false;
 	}
 
@@ -414,9 +449,6 @@ bool FVisualStudioSourceCodeAccessor::OpenVisualStudioSolutionViaDTE()
 		break;
 	}
 
-	// Uninitialize the com library, if we initialized it above (don't call if S_FALSE)
-	FWindowsPlatformMisc::CoUninitialize();
-
 	return bSuccess;
 }
 
@@ -424,8 +456,8 @@ bool FVisualStudioSourceCodeAccessor::OpenVisualStudioFilesInternalViaDTE(const
 {
 	ISourceCodeAccessModule& SourceCodeAccessModule = FModuleManager::LoadModuleChecked<ISourceCodeAccessModule>(TEXT("SourceCodeAccess"));
 
-	// Initialize the com library, if not already by this thread
-	if (!FWindowsPlatformMisc::CoInitialize())
+	FCoInitializeScope CoInitialze;
+	if (!CoInitialze.IsValid())
 	{
 		UE_LOG(LogVSAccessor, Error, TEXT( "ERROR - Could not initialize COM library!" ));
 		return false;
@@ -603,23 +635,19 @@ bool FVisualStudioSourceCodeAccessor::OpenVisualStudioFilesInternalViaDTE(const
 		}
 	}
 
-	// Uninitialize the com library, if we initialized it above (don't call if S_FALSE)
-	FWindowsPlatformMisc::CoUninitialize();
-
 	bWasDeferred = bDefer;
 	return bSuccess;
 }
 
 bool FVisualStudioSourceCodeAccessor::SaveAllOpenDocuments() const
 {
-	bool bSuccess = false;
-
-	// Initialize the com library, if not already by this thread
-	if (!FWindowsPlatformMisc::CoInitialize())
+	FCoInitializeScope CoInitialze;
+	if (!CoInitialze.IsValid())
 	{
-		UE_LOG(LogVSAccessor, Error, TEXT( "ERROR - Could not initialize COM library!" ));
-		return bSuccess;
+		return false;
 	}
+
+	bool bSuccess = false;
 	
 	TComPtr<EnvDTE::_DTE> DTE;
 	const FString SolutionPath = GetSolutionPath();
@@ -642,9 +670,6 @@ bool FVisualStudioSourceCodeAccessor::SaveAllOpenDocuments() const
 		UE_LOG(LogVSAccessor, Warning, TEXT("Couldn't access Visual Studio"));
 	}
 
-	// Uninitialize the com library, if we initialized it above (don't call if S_FALSE)
-	FWindowsPlatformMisc::CoUninitialize();
-
 	return bSuccess;
 }
 
@@ -677,10 +702,9 @@ bool GetProcessCommandLine(const ::DWORD InProcessID, FString& OutCommandLine)
 {
 	check(InProcessID);
 
-	// Initialize the com library, if not already by this thread
-	if (!FWindowsPlatformMisc::CoInitialize())
+	FCoInitializeScope CoInitialze;
+	if (!CoInitialze.IsValid())
 	{
-		UE_LOG(LogVSAccessor, Error, TEXT("ERROR - Could not initialize COM library!"));
 		return false;
 	}
 
@@ -746,9 +770,6 @@ bool GetProcessCommandLine(const ::DWORD InProcessID, FString& OutCommandLine)
 		pLoc->Release();
 	}
 
-	// Uninitialize the com library, if we initialized it above (don't call if S_FALSE)
-	FWindowsPlatformMisc::CoUninitialize();
-
 	return bSuccess;
 }
 
@@ -1468,6 +1489,12 @@ void FVisualStudioSourceCodeAccessor::AddVisualStudioVersion(const int MajorVers
 
 void FVisualStudioSourceCodeAccessor::AddVisualStudioVersionUsingVisualStudioSetupAPI(int VersionNumber)
 {
+	FCoInitializeScope CoInitialze;
+	if (!CoInitialze.IsValid())
+	{
+		return;
+	}
+
 	// Try to create the CoCreate the class; if that fails, likely no instances are registered.
 	TComPtr<ISetupConfiguration2> Query;
 	HRESULT Result = CoCreateInstance(__uuidof(SetupConfiguration), nullptr, CLSCTX_ALL, __uuidof(ISetupConfiguration2), (LPVOID*)&Query);
