Downgrade Visual Studio DTE connection warning to log as it is non-fatal and will be retried

#rnx

diff --git a/Engine/Plugins/Developer/VisualStudioSourceCodeAccess/Source/VisualStudioSourceCodeAccess/Private/VisualStudioSourceCodeAccessor.cpp b/Engine/Plugins/Developer/VisualStudioSourceCodeAccess/Source/VisualStudioSourceCodeAccess/Private/VisualStudioSourceCodeAccessor.cpp
--- a/Engine/Plugins/Developer/VisualStudioSourceCodeAccess/Source/VisualStudioSourceCodeAccess/Private/VisualStudioSourceCodeAccessor.cpp
+++ b/Engine/Plugins/Developer/VisualStudioSourceCodeAccess/Source/VisualStudioSourceCodeAccess/Private/VisualStudioSourceCodeAccessor.cpp
@@ -272,7 +272,7 @@ TArray<FString> RetrieveOpenedVisualStudioSolutionsViaDTE(const TArray<FVisualSt
 								}
 								else
 								{
-									UE_LOG(LogVSAccessor, Warning, TEXT("Visual Studio is open but could not be queried - it may be blocked by a modal operation"));
+									UE_LOG(LogVSAccessor, Log, TEXT("Visual Studio is open but could not be queried - it may still be initializing or blocked by a modal operation"));
 								}
 							}
 							else
@@ -359,7 +359,7 @@ EAccessVisualStudioResult AccessVisualStudioViaDTE(TComPtr<EnvDTE::_DTE>& OutDTE
 								}
 								else
 								{
-									UE_LOG(LogVSAccessor, Warning, TEXT("Visual Studio is open but could not be queried - it may be blocked by a modal operation"));
+									UE_LOG(LogVSAccessor, Log, TEXT("Visual Studio is open but could not be queried - it may still be initializing or blocked by a modal operation"));
 									AccessResult = EAccessVisualStudioResult::VSInstanceIsBlocked;
 								}
 							}
