Fix PVSStudio analysis warning

#rnx

diff --git a/Engine/Plugins/Developer/VisualStudioSourceCodeAccess/Source/VisualStudioSourceCodeAccess/Private/VisualStudioSourceCodeAccessor.cpp b/Engine/Plugins/Developer/VisualStudioSourceCodeAccess/Source/VisualStudioSourceCodeAccess/Private/VisualStudioSourceCodeAccessor.cpp
--- a/Engine/Plugins/Developer/VisualStudioSourceCodeAccess/Source/VisualStudioSourceCodeAccess/Private/VisualStudioSourceCodeAccessor.cpp
+++ b/Engine/Plugins/Developer/VisualStudioSourceCodeAccess/Source/VisualStudioSourceCodeAccess/Private/VisualStudioSourceCodeAccessor.cpp
@@ -539,9 +539,9 @@ bool FVisualStudioSourceCodeAccessor::OpenVisualStudioFilesInternalViaDTE(const
 
 							} while (--RetryCount >= 0);
 
-							if (Selection.IsValid() && !SUCCEEDED(Selection->MoveToLineAndOffset(Request.LineNumber, Request.ColumnNumber, false)))
+							if (Selection.IsValid() && !SUCCEEDED(Selection->MoveToLineAndOffset(Request.LineNumber, Request.ColumnNumber, VARIANT_FALSE)))
 							{
-								if (!SUCCEEDED(Selection->GotoLine(Request.LineNumber, true)))
+								if (!SUCCEEDED(Selection->GotoLine(Request.LineNumber, VARIANT_TRUE)))
 								{
 									UE_LOG(LogVSAccessor, Warning, TEXT("Couldn't goto column number '%i' of line '%i' in '%s'."), Request.ColumnNumber, Request.LineNumber, *Request.FullPath);
 								}
