VisualStudioSourceCodeAccess: Add option to open .uproject directly in visual studio as well as an option to prefer preview releases

diff --git a/Engine/Plugins/Developer/VisualStudioSourceCodeAccess/Source/VisualStudioSourceCodeAccess/Private/VisualStudioSourceCodeAccessModule.cpp b/Engine/Plugins/Developer/VisualStudioSourceCodeAccess/Source/VisualStudioSourceCodeAccess/Private/VisualStudioSourceCodeAccessModule.cpp
--- a/Engine/Plugins/Developer/VisualStudioSourceCodeAccess/Source/VisualStudioSourceCodeAccess/Private/VisualStudioSourceCodeAccessModule.cpp
+++ b/Engine/Plugins/Developer/VisualStudioSourceCodeAccess/Source/VisualStudioSourceCodeAccess/Private/VisualStudioSourceCodeAccessModule.cpp
@@ -1,9 +1,11 @@
 // Copyright Epic Games, Inc. All Rights Reserved.
 
 #include "VisualStudioSourceCodeAccessModule.h"
+#include "VisualStudioSourceCodeAccessSettings.h"
 #include "Features/IModularFeatures.h"
 #include "Modules/ModuleManager.h"
 #include "HAL/LowLevelMemTracker.h"
+#include "ISettingsModule.h"
 
 IMPLEMENT_MODULE( FVisualStudioSourceCodeAccessModule, VisualStudioSourceCodeAccess );
 
@@ -28,6 +30,18 @@ void FVisualStudioSourceCodeAccessModule::StartupModule()
 
 	// Bind our source control provider to the editor
 	IModularFeatures::Get().RegisterModularFeature(TEXT("SourceCodeAccessor"), &VisualStudioSourceCodeAccessor.Get() );
+
+	// Register to display our settings
+	ISettingsModule* SettingsModule = FModuleManager::GetModulePtr<ISettingsModule>("Settings");
+
+	if (SettingsModule != nullptr)
+	{
+		SettingsModule->RegisterSettings("Editor", "Plugins", "Visual Studio",
+			LOCTEXT("TargetSettingsName", "Visual Studio"),
+			LOCTEXT("TargetSettingsDescription", "Control how Visual Studio accesses source code, when set as the source code accessor."),
+			GetMutableDefault<UVisualStudioSourceCodeAccessSettings>()
+		);
+	}
 }
 
 void FVisualStudioSourceCodeAccessModule::ShutdownModule()
diff --git a/Engine/Plugins/Developer/VisualStudioSourceCodeAccess/Source/VisualStudioSourceCodeAccess/Private/VisualStudioSourceCodeAccessSettings.cpp b/Engine/Plugins/Developer/VisualStudioSourceCodeAccess/Source/VisualStudioSourceCodeAccess/Private/VisualStudioSourceCodeAccessSettings.cpp
new file mode 100644
--- /dev/null
+++ b/Engine/Plugins/Developer/VisualStudioSourceCodeAccess/Source/VisualStudioSourceCodeAccess/Private/VisualStudioSourceCodeAccessSettings.cpp
@@ -0,0 +1,10 @@
+// Copyright Epic Games, Inc. All Rights Reserved.
+
+#include "VisualStudioSourceCodeAccessSettings.h"
+
+#include UE_INLINE_GENERATED_CPP_BY_NAME(VisualStudioSourceCodeAccessSettings)
+
+UVisualStudioSourceCodeAccessSettings::UVisualStudioSourceCodeAccessSettings(const FObjectInitializer& ObjectInitializer)
+	: Super(ObjectInitializer)
+{
+}
diff --git a/Engine/Plugins/Developer/VisualStudioSourceCodeAccess/Source/VisualStudioSourceCodeAccess/Private/VisualStudioSourceCodeAccessSettings.h b/Engine/Plugins/Developer/VisualStudioSourceCodeAccess/Source/VisualStudioSourceCodeAccess/Private/VisualStudioSourceCodeAccessSettings.h
new file mode 100644
--- /dev/null
+++ b/Engine/Plugins/Developer/VisualStudioSourceCodeAccess/Source/VisualStudioSourceCodeAccess/Private/VisualStudioSourceCodeAccessSettings.h
@@ -0,0 +1,23 @@
+// Copyright Epic Games, Inc. All Rights Reserved.
+
+#pragma once
+
+#include "CoreMinimal.h"
+#include "UObject/ObjectMacros.h"
+#include "UObject/Object.h"
+#include "VisualStudioSourceCodeAccessSettings.generated.h"
+
+UCLASS(config=EditorSettings)
+class UVisualStudioSourceCodeAccessSettings : public UObject
+{
+	GENERATED_UCLASS_BODY()
+
+public:
+	/** Open a .uproject directly in Visual Studio instead of the generated .sln solution. */
+	UPROPERTY(Config, EditAnywhere, Category = "Visual Studio Source Code", meta = (DisplayName = "Direct Unreal Project Support"))
+	bool bUproject;
+
+	/** Prefer preview releases of Visual Studio. */
+	UPROPERTY(Config, EditAnywhere, Category = "Visual Studio Source Code", meta = (DisplayName = "Prefer Preview Releases"))
+	bool bPreview;
+};
diff --git a/Engine/Plugins/Developer/VisualStudioSourceCodeAccess/Source/VisualStudioSourceCodeAccess/Private/VisualStudioSourceCodeAccessor.cpp b/Engine/Plugins/Developer/VisualStudioSourceCodeAccess/Source/VisualStudioSourceCodeAccess/Private/VisualStudioSourceCodeAccessor.cpp
--- a/Engine/Plugins/Developer/VisualStudioSourceCodeAccess/Source/VisualStudioSourceCodeAccess/Private/VisualStudioSourceCodeAccessor.cpp
+++ b/Engine/Plugins/Developer/VisualStudioSourceCodeAccess/Source/VisualStudioSourceCodeAccess/Private/VisualStudioSourceCodeAccessor.cpp
@@ -2,6 +2,7 @@
 
 #include "VisualStudioSourceCodeAccessor.h"
 #include "VisualStudioSourceCodeAccessModule.h"
+#include "VisualStudioSourceCodeAccessSettings.h"
 #include "ISourceCodeAccessModule.h"
 #include "Modules/ModuleManager.h"
 #include "IDesktopPlatform.h"
@@ -118,14 +119,6 @@ int32 GetVisualStudioVersionForCompiler()
 {
 #if _MSC_VER >= 1930
 	return 17; // Visual Studio 2022
-#elif _MSC_VER >= 1920
-	return 16; // Visual Studio 2019
-#elif _MSC_VER >= 1910
-	return 15; // Visual Studio 2017
-#elif _MSC_VER == 1900
-	return 14; // Visual Studio 2015
-#elif _MSC_VER == 1800
-	return 12; // Visual Studio 2013
 #elif defined(__clang__)
 	return 0;
 #else
@@ -136,6 +129,11 @@ int32 GetVisualStudioVersionForCompiler()
 
 int32 GetVisualStudioVersionForSolution(const FString& InSolutionFile)
 {
+	if (FPaths::DirectoryExists(InSolutionFile))
+	{
+		// .uproject support uses a directory root instead of a file, and has no solution to check so defer to the version based on the compiler
+		return 0;
+	}
 	static const FString VisualStudioVersionString = TEXT("# Visual Studio ");
 
 	FString SolutionFileContents;
@@ -193,13 +191,8 @@ void FVisualStudioSourceCodeAccessor::RefreshAvailability()
 	FCoInitializeScope CoInitialze;
 	if (CoInitialze.IsValid())
 	{
-		AddVisualStudioVersionUsingVisualStudioSetupAPI(17); // Visual Studio 2022
-		AddVisualStudioVersionUsingVisualStudioSetupAPI(16); // Visual Studio 2019
-		AddVisualStudioVersionUsingVisualStudioSetupAPI(15); // Visual Studio 2017
+		AddVisualStudioVersionUsingVisualStudioSetupAPI(MinimumVisualStudioVersion);
 	}
-
-	AddVisualStudioVersion(14); // Visual Studio 2015
-	AddVisualStudioVersion(12); // Visual Studio 2013
 }
 
 void FVisualStudioSourceCodeAccessor::Shutdown()
@@ -1010,7 +1003,7 @@ bool FVisualStudioSourceCodeAccessor::OpenSolutionAtPath(const FString& InSoluti
 bool FVisualStudioSourceCodeAccessor::DoesSolutionExist() const
 {
 	const FString SolutionPath = GetSolutionPath();
-	return FPaths::FileExists(SolutionPath);
+	return FPaths::FileExists(SolutionPath) || FPaths::DirectoryExists(SolutionPath);
 }
 
 bool FVisualStudioSourceCodeAccessor::OpenVisualStudioFilesInternal(const TArray<FileOpenRequest>& Requests)
@@ -1081,39 +1074,71 @@ FString FVisualStudioSourceCodeAccessor::RetrieveSolutionForFileOpenRequests(con
 		FString SolutionFile;
 	};
 
+	const bool bUseUProject = GetDefault<UVisualStudioSourceCodeAccessSettings>()->bUproject;
+
 	TArray<FSolutionInfo> SolutionFileInfos;
 	for (const FileOpenRequest& Request : Requests)
 	{
 		FString CurrentPath = FPaths::GetPath(Request.FullPath);
 		while (!CurrentPath.IsEmpty())
 		{
-			TArray<FString> FilesInDirectory;
-			IFileManager::Get().FindFiles(FilesInDirectory, *CurrentPath, TEXT(".sln"));
-			IFileManager::Get().FindFiles(FilesInDirectory, *CurrentPath, TEXT(".slnf"));
-			for (const FString& FileInDirectory : FilesInDirectory)
+			if (bUseUProject)
 			{
-				FString AbsoluteFileName = CurrentPath / FileInDirectory;
-				FPaths::NormalizeFilename(AbsoluteFileName);
-
-				FSolutionInfo* SolutionInfo = SolutionFileInfos.FindByPredicate([&AbsoluteFileName](const FSolutionInfo& InSolutionInfo) { return InSolutionInfo.SolutionFile == AbsoluteFileName; });
-				if (SolutionInfo == nullptr)
+				TArray<FString> FilesInDirectory;
+				IFileManager::Get().FindFiles(FilesInDirectory, *CurrentPath, TEXT(".uproject"));
+				if (FilesInDirectory.Num() > 0)
 				{
-					SolutionInfo = &SolutionFileInfos.Add_GetRef(FSolutionInfo(AbsoluteFileName));
-
-					for (const FString& OpenedSolution : CurrentlyOpenedSolutions)
+					FSolutionInfo* SolutionInfo = SolutionFileInfos.FindByPredicate([&CurrentPath](const FSolutionInfo& InSolutionInfo) { return InSolutionInfo.SolutionFile == CurrentPath; });
+					if (SolutionInfo == nullptr)
 					{
-						if (OpenedSolution.Equals(AbsoluteFileName, ESearchCase::IgnoreCase))
+						SolutionInfo = &SolutionFileInfos.Add_GetRef(FSolutionInfo(CurrentPath));
+
+						for (const FString& OpenedSolution : CurrentlyOpenedSolutions)
 						{
-							SolutionInfo->OpenedSolutionState = FSolutionInfo::EOpenedSolutionState::CurrentlyOpenedExactMatch;
-							break;
+							if (OpenedSolution.Equals(CurrentPath, ESearchCase::IgnoreCase))
+							{
+								SolutionInfo->OpenedSolutionState = FSolutionInfo::EOpenedSolutionState::CurrentlyOpenedExactMatch;
+								break;
+							}
+							if (FPaths::GetCleanFilename(OpenedSolution).Equals(CurrentPath, ESearchCase::IgnoreCase))
+							{
+								SolutionInfo->OpenedSolutionState = FSolutionInfo::EOpenedSolutionState::CurrentlyOpened;
+							}
 						}
-						if (FPaths::GetCleanFilename(OpenedSolution).Equals(FileInDirectory, ESearchCase::IgnoreCase))
+					}
+					++SolutionInfo->RefCount;
+				}
+			}
+			else
+			{
+				TArray<FString> FilesInDirectory;
+				IFileManager::Get().FindFiles(FilesInDirectory, *CurrentPath, TEXT(".sln"));
+				IFileManager::Get().FindFiles(FilesInDirectory, *CurrentPath, TEXT(".slnf"));
+				for (const FString& FileInDirectory : FilesInDirectory)
+				{
+					FString AbsoluteFileName = CurrentPath / FileInDirectory;
+					FPaths::NormalizeFilename(AbsoluteFileName);
+
+					FSolutionInfo* SolutionInfo = SolutionFileInfos.FindByPredicate([&AbsoluteFileName](const FSolutionInfo& InSolutionInfo) { return InSolutionInfo.SolutionFile == AbsoluteFileName; });
+					if (SolutionInfo == nullptr)
+					{
+						SolutionInfo = &SolutionFileInfos.Add_GetRef(FSolutionInfo(AbsoluteFileName));
+
+						for (const FString& OpenedSolution : CurrentlyOpenedSolutions)
 						{
-							SolutionInfo->OpenedSolutionState = FSolutionInfo::EOpenedSolutionState::CurrentlyOpened;
+							if (OpenedSolution.Equals(AbsoluteFileName, ESearchCase::IgnoreCase))
+							{
+								SolutionInfo->OpenedSolutionState = FSolutionInfo::EOpenedSolutionState::CurrentlyOpenedExactMatch;
+								break;
+							}
+							if (FPaths::GetCleanFilename(OpenedSolution).Equals(FileInDirectory, ESearchCase::IgnoreCase))
+							{
+								SolutionInfo->OpenedSolutionState = FSolutionInfo::EOpenedSolutionState::CurrentlyOpened;
+							}
 						}
 					}
+					++SolutionInfo->RefCount;
 				}
-				++SolutionInfo->RefCount;
 			}
 			CurrentPath = FPaths::GetPath(CurrentPath);
 		}
@@ -1487,7 +1512,7 @@ void FVisualStudioSourceCodeAccessor::AddVisualStudioVersion(const int MajorVers
 	}
 }
 
-void FVisualStudioSourceCodeAccessor::AddVisualStudioVersionUsingVisualStudioSetupAPI(int VersionNumber)
+void FVisualStudioSourceCodeAccessor::AddVisualStudioVersionUsingVisualStudioSetupAPI(int MinimumVersionNumber)
 {
 	FCoInitializeScope CoInitialze;
 	if (!CoInitialze.IsValid())
@@ -1513,11 +1538,6 @@ void FVisualStudioSourceCodeAccessor::AddVisualStudioVersionUsingVisualStudioSet
 		return;
 	}
 
-	// Get the verison prefix string that we're looking for. We can compare this against reported version numbers from installed Visual Studio instances.
-	TCHAR VersionPrefix[10];
-	FCString::Sprintf(VersionPrefix, TEXT("%d."), VersionNumber);
-	size_t VersionPrefixLen = FCString::Strlen(VersionPrefix);
-
 	// Check the state and version of the enumerated instances
 	TComPtr<ISetupInstance> Instance;
 	for (;;)
@@ -1541,7 +1561,9 @@ void FVisualStudioSourceCodeAccessor::AddVisualStudioVersionUsingVisualStudioSet
 				Result = Instance2->GetInstallationVersion(&InstallationVersion);
 				if (SUCCEEDED(Result))
 				{
-					if (FCString::Strncmp(InstallationVersion, VersionPrefix, VersionPrefixLen) == 0)
+					int Major, Minor, Build, Patch;
+					int ScanResult = swscanf_s(InstallationVersion, TEXT("%d.%d.%d.%d"), &Major, &Minor, &Build, &Patch);
+					if (ScanResult == 4 && Major >= MinimumVersionNumber)
 					{
 						BSTR InstallationPath;
 						Result = Instance2->GetInstallationPath(&InstallationPath);
@@ -1552,11 +1574,11 @@ void FVisualStudioSourceCodeAccessor::AddVisualStudioVersionUsingVisualStudioSet
 							if (SUCCEEDED(Result))
 							{
 								VisualStudioLocation NewLocation;
-								NewLocation.VersionNumber = VersionNumber;
+								NewLocation.VersionNumber = Major;
 								NewLocation.bPreviewRelease = false;
 								NewLocation.ExecutablePath = FString::Printf(TEXT("%s\\%s"), InstallationPath, ProductPath);
 #if WITH_VISUALSTUDIO_DTE
-								NewLocation.ROTMoniker = FString::Printf(TEXT("!VisualStudio.DTE.%d.0"), VersionNumber);
+								NewLocation.ROTMoniker = FString::Printf(TEXT("!VisualStudio.DTE.%d.0"), Major);
 #endif
 
 								TComPtr<ISetupInstanceCatalog> Catalog;
@@ -1585,25 +1607,31 @@ void FVisualStudioSourceCodeAccessor::AddVisualStudioVersionUsingVisualStudioSet
 	}
 }
 
+int32 FVisualStudioSourceCodeAccessor::VisualStudioVersionSortWeight(const VisualStudioLocation& InLocation, const int32 ExactVersion, const bool bPreferPreview) const
+{
+	return
+		// First sort by VersionNumber. If the version matches ExactVersion it should be prioritized
+		(InLocation.VersionNumber == ExactVersion ? InLocation.VersionNumber * 100 : InLocation.VersionNumber * 10)
+		// Then by if a preview release is preferred or not
+		+ (InLocation.bPreviewRelease == bPreferPreview ? 1 : 0);
+}
+
 TArray<FVisualStudioSourceCodeAccessor::VisualStudioLocation> FVisualStudioSourceCodeAccessor::GetPrioritizedVisualStudioVersions(const FString& InSolution) const
 {
 	TArray<VisualStudioLocation> PrioritizedLocations = Locations;
-	
+
 	int32 SolutionVersion = GetVisualStudioVersionForSolution(InSolution);
 	if (SolutionVersion == 0)
 	{
 		SolutionVersion = GetVisualStudioVersionForCompiler();
 	}
 
-	if (SolutionVersion != 0)
+	const bool bPreferPreview = GetDefault<UVisualStudioSourceCodeAccessSettings>()->bPreview;
+
+	PrioritizedLocations.StableSort([&](const VisualStudioLocation& InFirst, const VisualStudioLocation& InSecond) -> bool
 	{
-		PrioritizedLocations.StableSort([&](const VisualStudioLocation& InFirst, const VisualStudioLocation& InSecond) -> bool
-		{
-			const int32 FirstSortWeight = (InFirst.VersionNumber == SolutionVersion) ? (InFirst.bPreviewRelease? 1 : 2) : 0;
-			const int32 SecondSortWeight = (InSecond.VersionNumber == SolutionVersion) ? (InSecond.bPreviewRelease? 1 : 2) : 0;
-			return FirstSortWeight >= SecondSortWeight;
-		});
-	}
+		return VisualStudioVersionSortWeight(InFirst, SolutionVersion, bPreferPreview) >= VisualStudioVersionSortWeight(InSecond, SolutionVersion, bPreferPreview);
+	});
 	
 	return PrioritizedLocations;
 }
@@ -1617,7 +1645,7 @@ bool FVisualStudioSourceCodeAccessor::RunVisualStudioAndOpenSolutionAndFiles(con
 	// Only open the solution if it exists
 	if (!SolutionPath.IsEmpty())
 	{
-		if (FPaths::FileExists(SolutionPath))
+		if (FPaths::FileExists(SolutionPath) || FPaths::DirectoryExists(SolutionPath))
 		{
 			Params += TEXT("\"");
 			Params += SolutionPath;
@@ -1692,12 +1720,34 @@ FString FVisualStudioSourceCodeAccessor::GetSolutionPath() const
 
 	if(IsInGameThread())
 	{
+		const bool bUseUProject = GetDefault<UVisualStudioSourceCodeAccessSettings>()->bUproject;
 		if (CachedSolutionPathOverride.Len() > 0)
 		{
-			CachedSolutionPath = CachedSolutionPathOverride + TEXT(".sln");
+			if (bUseUProject)
+			{
+				CachedSolutionPath = FPaths::GetPath(CachedSolutionPathOverride);
+				FPaths::NormalizeDirectoryName(CachedSolutionPath);
+			}
+			else
+			{
+				CachedSolutionPath = CachedSolutionPathOverride + TEXT(".sln");
+				FPaths::NormalizeFilename(CachedSolutionPath);
+			}
 		}
 		else
 		{
+			if (bUseUProject)
+			{
+				const FProjectDescriptor* CurrentProject = IProjectManager::Get().GetCurrentProject();
+				if (CurrentProject != nullptr)
+				{
+					// VS support is implemented to open the directory that contains the .uproject
+					CachedSolutionPath = FPaths::ConvertRelativePathToFull(FPaths::ProjectDir());
+					FPaths::NormalizeDirectoryName(CachedSolutionPath);
+					return CachedSolutionPath;
+				}
+			}
+			
 			FString PrimaryProjectPath;
 			if (!FFileHelper::LoadFileToString(PrimaryProjectPath, *(FPaths::EngineIntermediateDir() / TEXT("ProjectFiles") / TEXT("PrimaryProjectPath.txt"))))
 			{
@@ -1715,6 +1765,7 @@ FString FVisualStudioSourceCodeAccessor::GetSolutionPath() const
 				}
 			}
 			CachedSolutionPath = PrimaryProjectPath + TEXT(".sln");
+			FPaths::NormalizeFilename(CachedSolutionPath);
 		}
 	}
 
diff --git a/Engine/Plugins/Developer/VisualStudioSourceCodeAccess/Source/VisualStudioSourceCodeAccess/Private/VisualStudioSourceCodeAccessor.h b/Engine/Plugins/Developer/VisualStudioSourceCodeAccess/Source/VisualStudioSourceCodeAccess/Private/VisualStudioSourceCodeAccessor.h
--- a/Engine/Plugins/Developer/VisualStudioSourceCodeAccess/Source/VisualStudioSourceCodeAccess/Private/VisualStudioSourceCodeAccessor.h
+++ b/Engine/Plugins/Developer/VisualStudioSourceCodeAccess/Source/VisualStudioSourceCodeAccess/Private/VisualStudioSourceCodeAccessor.h
@@ -99,7 +99,7 @@ private:
 	/** 
 	 * Add a new version of Visual Studio to the supported locations array
 	 * 
-	 * @param	MajorVersion	The major version number of Visual Studio (eg, 11 for VS2012, 12 for VS2013)
+	 * @param	MajorVersion	The major version number of Visual Studio (eg, 17 for VS2022)
 	 * @param	bAllowExpress	true to also add the express version of Visual Studio to the list of locations
 	 */
 	void AddVisualStudioVersion(const int MajorVersion, const bool bAllowExpress = true);
@@ -107,9 +107,20 @@ private:
 	/**
 	 * Add a new version of Visual Studio to the supported locations array using the Visual Studio Setup API
 	 *
-	 * @param	MajorVersion	The major version number of Visual Studio (eg, 16 for VS2019)
+	 * @param	MinimumVersionNumber	The minimum major version number of Visual Studio (eg, 17 for VS2022) to return
 	 */
-	void AddVisualStudioVersionUsingVisualStudioSetupAPI(const int MajorVersion);
+	void AddVisualStudioVersionUsingVisualStudioSetupAPI(const int MinimumVersionNumber);
+
+	/**
+	 * Get the sort weight of a Visual Studio install when ordering by priority. 
+	 * 
+	 * @param	InLocation		The Visual Studio install
+	 * @param	ExactVersion	The specific major version number of Visual Studio to prioritize
+	 * @param	bPreferPreview	true if preview releases should be prioritized
+	 *
+	 * @return	The sort weight of an install location
+	 */
+	int32 VisualStudioVersionSortWeight(const VisualStudioLocation& InLocation, const int32 ExactVersion, const bool bPreferPreview) const;
 
 	/** 
 	 * Get the prioritized list of VS install locations based upon the version of the given solution
@@ -147,6 +158,9 @@ private:
 
 private:
 
+	/** The minimum major version of Visual Studio we support (eg, 17 for Visual Studio 2022) */
+	static constexpr int32 MinimumVisualStudioVersion = 17;
+
 	/** The versions of VS we support, in preference order */
 	TArray<VisualStudioLocation> Locations;
 
@@ -157,7 +171,7 @@ private:
 	mutable FString CachedSolutionPathOverride;
 
 	/** If !0 it represents the time at which the a VS instance was opened */
-	double	VSLaunchTime;
+	double VSLaunchTime;
 
 	/** Critical section for updating DeferredRequests */
 	FCriticalSection DeferredRequestsCriticalSection;
diff --git a/Engine/Plugins/Developer/VisualStudioSourceCodeAccess/Source/VisualStudioSourceCodeAccess/VisualStudioSourceCodeAccess.Build.cs b/Engine/Plugins/Developer/VisualStudioSourceCodeAccess/Source/VisualStudioSourceCodeAccess/VisualStudioSourceCodeAccess.Build.cs
--- a/Engine/Plugins/Developer/VisualStudioSourceCodeAccess/Source/VisualStudioSourceCodeAccess/VisualStudioSourceCodeAccess.Build.cs
+++ b/Engine/Plugins/Developer/VisualStudioSourceCodeAccess/Source/VisualStudioSourceCodeAccess/VisualStudioSourceCodeAccess.Build.cs
@@ -12,6 +12,7 @@ namespace UnrealBuildTool.Rules
 				new string[]
 				{
 					"Core",
+					"CoreUObject",
 					"SourceCodeAccess",
 					"DesktopPlatform",
 					"Projects",
