Remove experimental from Visual Studio 2026 support

#rnx
#Synced-CL 51561575

diff --git a/Engine/Plugins/Developer/VisualStudioSourceCodeAccess/Source/VisualStudioSourceCodeAccess/Private/VisualStudioSourceCodeAccessModule.cpp b/Engine/Plugins/Developer/VisualStudioSourceCodeAccess/Source/VisualStudioSourceCodeAccess/Private/VisualStudioSourceCodeAccessModule.cpp
--- a/Engine/Plugins/Developer/VisualStudioSourceCodeAccess/Source/VisualStudioSourceCodeAccess/Private/VisualStudioSourceCodeAccessModule.cpp
+++ b/Engine/Plugins/Developer/VisualStudioSourceCodeAccess/Source/VisualStudioSourceCodeAccess/Private/VisualStudioSourceCodeAccessModule.cpp
@@ -27,7 +27,7 @@ void FVisualStudioSourceCodeAccessModule::StartupModule()
 	// Add all the explicit version wrappers. If one of these is selected, UBT will generate project files in the appropriate format. Editor behavior is still to detect which version to use
 	// from the solution on disk.
 	RegisterWrapper("VisualStudio2022", LOCTEXT("VisualStudio2022", "Visual Studio 2022"), LOCTEXT("UsingVisualStudio2022", "Open source code files in Visual Studio 2022"));
-	RegisterWrapper("VisualStudio2026", LOCTEXT("VisualStudio2026", "Visual Studio 2026 (experimental}"), LOCTEXT("UsingVisualStudio2026", "Open source code files in Visual Studio 2026"));
+	RegisterWrapper("VisualStudio2026", LOCTEXT("VisualStudio2026", "Visual Studio 2026"), LOCTEXT("UsingVisualStudio2026", "Open source code files in Visual Studio 2026"));
 
 	// Bind our source control provider to the editor
 	IModularFeatures::Get().RegisterModularFeature(TEXT("SourceCodeAccessor"), &VisualStudioSourceCodeAccessor.Get() );
